Removed a #pragma once from a cpp file to fix CIS (NonUnityTestCompile_OnMac)

diff --git a/Engine/Plugins/Media/AvfMedia/Source/AvfMedia/Private/Tracks/AvfMediaTrack.cpp b/Engine/Plugins/Media/AvfMedia/Source/AvfMedia/Private/Tracks/AvfMediaTrack.cpp
--- a/Engine/Plugins/Media/AvfMedia/Source/AvfMedia/Private/Tracks/AvfMediaTrack.cpp
+++ b/Engine/Plugins/Media/AvfMedia/Source/AvfMedia/Private/Tracks/AvfMediaTrack.cpp
@@ -1,7 +1,5 @@
 // Copyright 1998-2014 Epic Games, Inc. All Rights Reserved.
 
-#pragma once
-
 #include "AVFMediaPrivatePCH.h"
 
 
