AV Foundation Media Player Changes

- IOS and Mac should now work properly
- Changes made to IOS and Mac file paths. Some issues with accessing file names outside bundle have been fixed.
- Fix for Threading crash when calling the opened event not on the game thread.
- Some all around code tidiness done.

diff --git a/Engine/Plugins/Media/AvfMedia/Source/AvfMedia/Private/Player/AvfMediaPlayer.cpp b/Engine/Plugins/Media/AvfMedia/Source/AvfMedia/Private/Player/AvfMediaPlayer.cpp
--- a/Engine/Plugins/Media/AvfMedia/Source/AvfMedia/Private/Player/AvfMediaPlayer.cpp
+++ b/Engine/Plugins/Media/AvfMedia/Source/AvfMedia/Private/Player/AvfMediaPlayer.cpp
@@ -2,6 +2,9 @@
 
 #include "AvfMediaPrivatePCH.h"
 
+#if PLATFORM_MAC
+#include "CocoaThread.h"
+#endif
 
 /* FAvfMediaPlayer structors
  *****************************************************************************/
@@ -66,26 +69,41 @@ void FAvfMediaPlayer::Close()
 {
     CurrentTime = 0;
 	MediaUrl = FString();
+    
     if( PlayerItem != nil )
     {
-        [PlayerItem release];
         PlayerItem = nil;
     }
 
     if( MediaPlayer != nil )
     {
-        [MediaPlayer release];
         MediaPlayer = nil;
     }
+    
 	Tracks.Reset();
     
     Duration = CurrentTime = FTimespan::Zero();
     
-    ClosedEvent.Broadcast();
+#if PLATFORM_MAC
+    GameThreadCall(^{
+#elif PLATFORM_IOS
+    // Report back to the game thread whether this succeeded.
+    [FIOSAsyncTask CreateTaskWithBlock : ^ bool(void){
+#endif
+        // Displatch on the gamethread that we have closed the video.
+        ClosedEvent.Broadcast();
+         
+#if PLATFORM_IOS
+        return true;
+    }];
+#elif PLATFORM_MAC
+    });
+#endif
+
 }
 
 
-const IMediaInfo& FAvfMediaPlayer::GetMediaInfo() const 
+const IMediaInfo& FAvfMediaPlayer::GetMediaInfo() const
 {
 	return *this;
 }
@@ -97,7 +115,7 @@ float FAvfMediaPlayer::GetRate() const
 }
 
 
-FTimespan FAvfMediaPlayer::GetTime() const 
+FTimespan FAvfMediaPlayer::GetTime() const
 {
     return CurrentTime;
 }
@@ -123,7 +141,7 @@ bool FAvfMediaPlayer::IsPaused() const
 
 bool FAvfMediaPlayer::IsPlaying() const
 {
-    return (MediaPlayer != nil) && (1.0f == MediaPlayer.rate) && (CurrentTime <= Duration);
+    return (MediaPlayer != nil) && (1.0f == MediaPlayer.rate) && (Tracks.Num() > 0) && (CurrentTime <= Duration);
 }
 
 
@@ -154,51 +172,91 @@ bool FAvfMediaPlayer::IsReady() const
 
 bool FAvfMediaPlayer::Open( const FString& Url )
 {
+    // Cache off the url for use on other threads.
+    FString CachedUrl = Url;
+    
     Close();
     
     bool bSuccessfullyOpenedMovie = false;
 
-    MediaUrl = Url;
+    FString Ext = FPaths::GetExtension(Url);
+    FString BFNStr = FPaths::GetBaseFilename(Url);
     
-    MediaPlayer = [[AVPlayer alloc] init];
-    NSURL* nsMediaUrl = [NSURL fileURLWithPath:MediaUrl.GetNSString()];
-    if( MediaPlayer )
+#if PLATFORM_MAC
+    NSURL* nsMediaUrl = [NSURL fileURLWithPath:Url.GetNSString()];
+#elif PLATFORM_IOS
+    NSURL* nsMediaUrl = [[NSBundle mainBundle] URLForResource:BFNStr.GetNSString() withExtension:Ext.GetNSString()];
+#endif
+    
+    if( nsMediaUrl != nil )
     {
-        PlayerItem = [AVPlayerItem playerItemWithURL: nsMediaUrl];
-        if( PlayerItem != nil )
+        MediaUrl = FPaths::GetCleanFilename(Url);
+        
+        MediaPlayer = [[AVPlayer alloc] init];
+        
+        if( MediaPlayer )
         {
-            Duration = FTimespan::FromSeconds( CMTimeGetSeconds( PlayerItem.asset.duration ) );
-
-            [[PlayerItem asset] loadValuesAsynchronouslyForKeys: @[@"tracks"] completionHandler:^
+            PlayerItem = [AVPlayerItem playerItemWithURL: nsMediaUrl];
+            if( PlayerItem != nil )
             {
-                if( [[PlayerItem asset] statusOfValueForKey:@"tracks" error:nil] == AVKeyValueStatusLoaded )
+                Duration = FTimespan::FromSeconds( CMTimeGetSeconds( PlayerItem.asset.duration ) );
+
+                [[PlayerItem asset] loadValuesAsynchronouslyForKeys: @[@"tracks"] completionHandler:^
                 {
-                    NSArray* VideoTracks = [[PlayerItem asset] tracksWithMediaType:AVMediaTypeVideo];
-                    if( VideoTracks.count > 0 )
+                    NSError* Error = nil;
+                    if( [[PlayerItem asset] statusOfValueForKey:@"tracks" error:&Error] == AVKeyValueStatusLoaded )
                     {
-                        AVAssetTrack* VideoTrack = [VideoTracks objectAtIndex: 0];
-                        Tracks.Add( MakeShareable( new FAvfMediaVideoTrack(VideoTrack) ) );
+                        NSArray* VideoTracks = [[PlayerItem asset] tracksWithMediaType:AVMediaTypeVideo];
+                        if( VideoTracks.count > 0 )
+                        {
+                            AVAssetTrack* VideoTrack = [VideoTracks objectAtIndex: 0];
+                            Tracks.Add( MakeShareable( new FAvfMediaVideoTrack(VideoTrack) ) );
      
-                        OpenedEvent.Broadcast( MediaUrl );
+#if PLATFORM_MAC
+                            GameThreadCall(^{
+#elif PLATFORM_IOS
+                            // Report back to the game thread whether this succeeded.
+                            [FIOSAsyncTask CreateTaskWithBlock : ^ bool(void){
+#endif
+                                // Displatch on the gamethread that we have opened the video.
+                                OpenedEvent.Broadcast( CachedUrl );
+                         
+#if PLATFORM_IOS
+                                return true;
+                            }];
+#elif PLATFORM_MAC
+                            });
+#endif
+                        }
+                    }
+                    else
+                    {
+                        NSDictionary *userInfo = [Error userInfo];
+                        NSString *errstr = [[userInfo objectForKey : NSUnderlyingErrorKey] localizedDescription];
+                        UE_LOG(LogAvfMedia, Display, TEXT("Failed to load video tracks. [%s]"), *FString(errstr));
                     }
-                }
-            }];
+                }];
             
-            [MediaPlayer replaceCurrentItemWithPlayerItem: PlayerItem];
+                [MediaPlayer replaceCurrentItemWithPlayerItem: PlayerItem];
             
-            [[MediaPlayer currentItem] seekToTime:kCMTimeZero];
-            MediaPlayer.rate = 0.0;
+                [[MediaPlayer currentItem] seekToTime:kCMTimeZero];
+                MediaPlayer.rate = 0.0;
             
-            bSuccessfullyOpenedMovie = true;
+                bSuccessfullyOpenedMovie = true;
+            }
+            else
+            {
+                UE_LOG(LogAvfMedia, Warning, TEXT("Failed to open player item with Url:"), *Url);
+            }
         }
         else
         {
-            UE_LOG(LogAvfMedia, Warning, TEXT("Failed to open player item with Url:"), *Url);
+            UE_LOG(LogAvfMedia, Error, TEXT("Failed to create instance of an AVPlayer"));
         }
     }
     else
     {
-        UE_LOG(LogAvfMedia, Error, TEXT("Failed to create instance of an AVPlayer"));
+        UE_LOG(LogAvfMedia, Warning, TEXT("Failed to open Media file:"), *Url);
     }
     
     
@@ -250,22 +308,28 @@ bool FAvfMediaPlayer::SetRate( float Rate )
 /* FTickerBase interface
  *****************************************************************************/
 
+bool FAvfMediaPlayer::ShouldAdvanceFrames() const
+{
+    return IsReady() && IsPlaying();
+}
+
 bool FAvfMediaPlayer::Tick( float DeltaTime )
 {
-    CurrentTime = FTimespan::FromSeconds(CMTimeGetSeconds([[MediaPlayer currentItem] currentTime]));
-    
-    if (IsPlaying())
+    if( ShouldAdvanceFrames() )
     {
+        CurrentTime = FTimespan::FromSeconds(CMTimeGetSeconds([[MediaPlayer currentItem] currentTime]));
+        
         for (IMediaTrackRef& Track : Tracks)
         {
             FAvfMediaVideoTrack* AVFTrack = (FAvfMediaVideoTrack*)&Track.Get();
-            
-			if (AVFTrack != nil)
+        
+            if (AVFTrack != nil)
             {
                 AVFTrack->ReadFrameAtTime([[MediaPlayer currentItem] currentTime]);
-			}
+            }
         }
     }
+
     
     return true;
 }
diff --git a/Engine/Plugins/Media/AvfMedia/Source/AvfMedia/Private/Player/AvfMediaPlayer.h b/Engine/Plugins/Media/AvfMedia/Source/AvfMedia/Private/Player/AvfMediaPlayer.h
--- a/Engine/Plugins/Media/AvfMedia/Source/AvfMedia/Private/Player/AvfMediaPlayer.h
+++ b/Engine/Plugins/Media/AvfMedia/Source/AvfMedia/Private/Player/AvfMediaPlayer.h
@@ -67,6 +67,14 @@ public:
 
     virtual bool Tick( float DeltaTime ) override;
 
+
+private:
+
+    // FAVFMediaPlayer interface
+
+    /** Should the media player try and progress */
+    bool ShouldAdvanceFrames() const;
+
 private:
 
     /** The AVFoundation media player */
diff --git a/Engine/Plugins/Media/AvfMedia/Source/AvfMedia/Private/Tracks/AvfMediaVideoTrack.cpp b/Engine/Plugins/Media/AvfMedia/Source/AvfMedia/Private/Tracks/AvfMediaVideoTrack.cpp
--- a/Engine/Plugins/Media/AvfMedia/Source/AvfMedia/Private/Tracks/AvfMediaVideoTrack.cpp
+++ b/Engine/Plugins/Media/AvfMedia/Source/AvfMedia/Private/Tracks/AvfMediaVideoTrack.cpp
@@ -11,6 +11,8 @@ FAvfMediaVideoTrack::FAvfMediaVideoTrack( AVAssetTrack* InVideoTrack )
     , SyncStatus(Default)
     , BitRate(0)
 {
+    LatestSamples = nil;
+    
     NSError* nsError = nil;
     AVReader = [[AVAssetReader alloc] initWithAsset: [InVideoTrack asset] error:&nsError];
     if( nsError != nil )
@@ -18,28 +20,37 @@ FAvfMediaVideoTrack::FAvfMediaVideoTrack( AVAssetTrack* InVideoTrack )
         FString ErrorStr( [nsError localizedDescription] );
 		UE_LOG(LogAvfMedia, Error, TEXT("Failed to create asset reader: %s"), *ErrorStr);
     }
-    
-    // Initialize our video output to match the format of the texture we'll be streaming to.
-    NSMutableDictionary* OutputSettings = [NSMutableDictionary dictionary];
-    [OutputSettings setObject: [NSNumber numberWithInt:kCVPixelFormatType_32BGRA] forKey:(NSString*)kCVPixelBufferPixelFormatTypeKey];
-        
-    AVVideoOutput = [[AVAssetReaderTrackOutput alloc] initWithTrack:InVideoTrack outputSettings:OutputSettings];
-    AVVideoOutput.alwaysCopiesSampleData = NO;
+    else
+    {
+        // Initialize our video output to match the format of the texture we'll be streaming to.
+        NSMutableDictionary* OutputSettings = [NSMutableDictionary dictionary];
+        [OutputSettings setObject: [NSNumber numberWithInt:kCVPixelFormatType_32BGRA] forKey:(NSString*)kCVPixelBufferPixelFormatTypeKey];
         
-    [AVReader addOutput:AVVideoOutput];
-    [AVReader startReading];
-    
-    bVideoTracksLoaded = true;
+        AVVideoOutput = [[AVAssetReaderTrackOutput alloc] initWithTrack:InVideoTrack outputSettings:OutputSettings];
+        AVVideoOutput.alwaysCopiesSampleData = NO;
     
-    GetVideoDetails( InVideoTrack );
+        FrameRate = 1.0f / [[AVVideoOutput track] nominalFrameRate];
+        
+        [AVReader addOutput:AVVideoOutput];
+        bVideoTracksLoaded = [AVReader startReading];
+
+        // The initial read of a frame will allow us to gather information on the video track, before we start playing the samples
+        ReadFrameAtTime( kCMTimeZero, true );
+    }
 }
 
 
 FAvfMediaVideoTrack::~FAvfMediaVideoTrack()
 {
-    // Destroyed.
-    [AVReader release];
-    [AVVideoOutput release];
+    if( AVVideoOutput != nil )
+    {
+        AVVideoOutput = nil;
+    }
+    
+    if( AVReader != nil )
+    {
+        AVReader = nil;
+    }
 }
 
 
@@ -60,51 +71,6 @@ bool FAvfMediaVideoTrack::SeekToTime( const CMTime& SeekTime )
 }
 
 
-void FAvfMediaVideoTrack::GetVideoDetails( AVAssetTrack* InVideoTrack )
-{
-    // Setup the asset reader to ascertain some further video details.
-    NSError* nsError = nil;
-    AVAssetReader* DetailsAssetReader = [[AVAssetReader alloc] initWithAsset: [InVideoTrack asset] error:&nsError];
-    if( nsError != nil )
-    {
-        FString ErrorStr( [nsError localizedDescription] );
-		UE_LOG(LogAvfMedia, Error, TEXT("Failed to create asset reader: %s"), *ErrorStr);
-    }
-    
-    // Initialize our video output so we can get some details
-    NSMutableDictionary* OutputSettings = [NSMutableDictionary dictionary];
-    [OutputSettings setObject: [NSNumber numberWithInt:kCVPixelFormatType_32BGRA] forKey:(NSString*)kCVPixelBufferPixelFormatTypeKey];
-    
-    AVAssetReaderTrackOutput* DetailsTrackOutput = [[AVAssetReaderTrackOutput alloc] initWithTrack:InVideoTrack outputSettings:OutputSettings];
-    DetailsTrackOutput.alwaysCopiesSampleData = NO;
-    
-    
-    // Read Frame rate
-    FrameRate = 1.0f / [[DetailsTrackOutput track] nominalFrameRate];
-    
-    
-    // Let's start reading the buffer to get some further details.
-    [DetailsAssetReader addOutput:DetailsTrackOutput];
-    [DetailsAssetReader startReading];
-    
-    // Grab the pixel buffer and lock it for reading.
-    CMSampleBufferRef LatestSamples = [DetailsTrackOutput copyNextSampleBuffer];
-    CVImageBufferRef PixelBuffer = CMSampleBufferGetImageBuffer(LatestSamples);
-    CVPixelBufferLockBaseAddress( PixelBuffer, kCVPixelBufferLock_ReadOnly );
-    
-    // Grab the width and height
-    Width  = (uint32)CVPixelBufferGetWidth(PixelBuffer);
-    Height = (uint32)CVPixelBufferGetHeight(PixelBuffer);
-    
-    CVPixelBufferUnlockBaseAddress( PixelBuffer, kCVPixelBufferLock_ReadOnly );
-    CFRelease(LatestSamples);
-    LatestSamples= nil;
-    
-    [DetailsAssetReader cancelReading];
-    [DetailsAssetReader release];
-}
-
-
 void FAvfMediaVideoTrack::ResetAssetReader()
 {
     [AVReader cancelReading];
@@ -118,7 +84,7 @@ bool FAvfMediaVideoTrack::IsReady() const
 }
 
 
-void FAvfMediaVideoTrack::ReadFrameAtTime( const CMTime& AVPlayerTime )
+void FAvfMediaVideoTrack::ReadFrameAtTime( const CMTime& AVPlayerTime, bool bInIsInitialFrameRead )
 {
     check( bVideoTracksLoaded );
     
@@ -126,8 +92,6 @@ void FAvfMediaVideoTrack::ReadFrameAtTime( const CMTime& AVPlayerTime )
     {
         double CurrentAVTime = CMTimeGetSeconds( AVPlayerTime );
     
-        CMSampleBufferRef LatestSamples = nil;
-    
         // We need to synchronize the video playback with the audio.
         // If the video frame is within tolerance (Ready), update the Texture Data.
         // If the video is Behind, throw it away and get the next one until we catch up with the ref time.
@@ -194,7 +158,7 @@ void FAvfMediaVideoTrack::ReadFrameAtTime( const CMTime& AVPlayerTime )
     
             uint8* VideoDataBuffer = (uint8*)CVPixelBufferGetBaseAddress( PixelBuffer );
             uint32 BufferSize = Width * Height * 4;
-    
+            
             // Push the pixel data for the frame to the sinks which are listening for it.
             for( IMediaSinkWeakPtr& SinkPtr : Sinks )
             {
@@ -206,14 +170,19 @@ void FAvfMediaVideoTrack::ReadFrameAtTime( const CMTime& AVPlayerTime )
             }
     
             CVPixelBufferUnlockBaseAddress( PixelBuffer, kCVPixelBufferLock_ReadOnly );
-            CFRelease(LatestSamples);
-            LatestSamples= nil;
+            
+            // If we have are reading the initial frame, do not discard the latest samples.
+            if( bInIsInitialFrameRead == false )
+            {
+                CFRelease(LatestSamples);
+                LatestSamples= nil;
+            }
         }
     
         if( SyncStatus != Ahead )
         {
             // Reset status;
-            SyncStatus = Default;
+            SyncStatus = bInIsInitialFrameRead ? Ready : Default;
         }
     }
 }
diff --git a/Engine/Plugins/Media/AvfMedia/Source/AvfMedia/Private/Tracks/AvfMediaVideoTrack.h b/Engine/Plugins/Media/AvfMedia/Source/AvfMedia/Private/Tracks/AvfMediaVideoTrack.h
--- a/Engine/Plugins/Media/AvfMedia/Source/AvfMedia/Private/Tracks/AvfMediaVideoTrack.h
+++ b/Engine/Plugins/Media/AvfMedia/Source/AvfMedia/Private/Tracks/AvfMediaVideoTrack.h
@@ -83,8 +83,9 @@ public:
      * Read the track information for the frame at the given time. This should be provided by the audio manager to ensure the tracks are in sync.
      *
      * @param AVPlayerTime The time the AVPlayer is currently synced to.
+     * @param bInIsInitialFrameRead Flag whether this is the first frame read of the asset.
      */
-    void ReadFrameAtTime( const CMTime& AVPlayerTime );
+    void ReadFrameAtTime( const CMTime& AVPlayerTime, bool bInIsInitialFrameRead = false );
 
     /**
      * Set the Av Asset reader to a specified time.
@@ -102,13 +103,6 @@ public:
 
 private:
 
-    /**
-     * Set up the details interface.
-     *
-     * @param InVideoTrack The track for which we can gather the information.
-     */
-    void GetVideoDetails(AVAssetTrack* InVideoTrack);
-
     /** Reset the AVAssetReader to a different point, which will enable us to regrab frames for seeking and restarting. */
     void ResetAssetReader();
 
@@ -146,4 +140,7 @@ private:
 
     /** The collection of output streams we are writing to. */
     TArray<IMediaSinkWeakPtr> Sinks;
+
+    /** Access to the video samples */
+    CMSampleBufferRef LatestSamples;
 };
