AvfMedia: Include what you use

diff --git a/Engine/Plugins/Media/AvfMedia/Source/AvfMedia/Private/AvfMediaModule.cpp b/Engine/Plugins/Media/AvfMedia/Source/AvfMedia/Private/AvfMediaModule.cpp
--- a/Engine/Plugins/Media/AvfMedia/Source/AvfMedia/Private/AvfMediaModule.cpp
+++ b/Engine/Plugins/Media/AvfMedia/Source/AvfMedia/Private/AvfMediaModule.cpp
@@ -1,6 +1,10 @@
 // Copyright 1998-2014 Epic Games, Inc. All Rights Reserved.
 
 #include "AVFMediaPrivatePCH.h"
+#include "IMediaModule.h"
+#include "IMediaPlayerFactory.h"
+#include "ModuleInterface.h"
+#include "ModuleManager.h"
 
 
 DEFINE_LOG_CATEGORY(LogAvfMedia);
@@ -17,10 +21,8 @@ class FAvfMediaModule
 {
 public:
 
-	/**
-	 * Default constructor.
-	 */
-	FAvfMediaModule( )
+	/** Default constructor. */
+	FAvfMediaModule()
 		: Initialized(false)
 	{ }
 
@@ -28,7 +30,7 @@ public:
 
 	// IModuleInterface interface
 
-	virtual void StartupModule( ) override
+	virtual void StartupModule() override
 	{
 		// load required libraries
 		IMediaModule* MediaModule = FModuleManager::LoadModulePtr<IMediaModule>("Media");
@@ -52,7 +54,7 @@ public:
 		Initialized = true;
 	}
 
-	virtual void ShutdownModule( ) override
+	virtual void ShutdownModule() override
 	{
 		if (!Initialized)
 		{
@@ -75,7 +77,7 @@ public:
 
 	// IMediaPlayerFactory interface
 
-	virtual TSharedPtr<IMediaPlayer> CreatePlayer( ) override
+	virtual TSharedPtr<IMediaPlayer> CreatePlayer() override
 	{
 		if (Initialized)
 		{
@@ -85,7 +87,7 @@ public:
 		return nullptr;
 	}
 
-	virtual const FMediaFormats& GetSupportedFormats( ) const override
+	virtual const FMediaFormats& GetSupportedFormats() const override
 	{
 		return SupportedFormats;
 	}
@@ -97,17 +99,17 @@ protected:
 	 *
 	 * @return true on success, false otherwise.
 	 */
-	bool LoadRequiredLibraries( )
+	bool LoadRequiredLibraries()
     {
 		return true;
 	}
 
 private:
 
-	// Whether the module has been initialized.
+	/** Whether the module has been initialized. */
 	bool Initialized;
 
-	// The collection of supported media formats.
+	/** The collection of supported media formats. */
 	FMediaFormats SupportedFormats;
 };
 
diff --git a/Engine/Plugins/Media/AvfMedia/Source/AvfMedia/Private/AvfMediaPrivatePCH.h b/Engine/Plugins/Media/AvfMedia/Source/AvfMedia/Private/AvfMediaPrivatePCH.h
--- a/Engine/Plugins/Media/AvfMedia/Source/AvfMedia/Private/AvfMediaPrivatePCH.h
+++ b/Engine/Plugins/Media/AvfMedia/Source/AvfMedia/Private/AvfMediaPrivatePCH.h
@@ -3,14 +3,19 @@
 #pragma once
 
 
-/* Dependencies
+/* Private dependencies
  *****************************************************************************/
 
-#include "ModuleManager.h"
-#include "Media.h"
+#include "Core.h"
+#include "IMediaSink.h"
+#include "IMediaTrack.h"
+#include "IMediaTrackAudioDetails.h"
+#include "IMediaTrackCaptionDetails.h"
+#include "IMediaTrackVideoDetails.h"
 
 #import <AVFoundation/AVFoundation.h>
 
+
 /* Private macros
  *****************************************************************************/
 
diff --git a/Engine/Plugins/Media/AvfMedia/Source/AvfMedia/Private/Player/AvfMediaPlayer.cpp b/Engine/Plugins/Media/AvfMedia/Source/AvfMedia/Private/Player/AvfMediaPlayer.cpp
--- a/Engine/Plugins/Media/AvfMedia/Source/AvfMedia/Private/Player/AvfMediaPlayer.cpp
+++ b/Engine/Plugins/Media/AvfMedia/Source/AvfMedia/Private/Player/AvfMediaPlayer.cpp
@@ -6,6 +6,7 @@
 #include "CocoaThread.h"
 #endif
 
+
 /* FAvfMediaPlayer structors
  *****************************************************************************/
 
@@ -169,7 +170,6 @@ bool FAvfMediaPlayer::IsReady() const
 /* FAvfMediaPlayer implementation
  *****************************************************************************/
 
-
 bool FAvfMediaPlayer::Open( const FString& Url )
 {
     // Cache off the url for use on other threads.
@@ -308,11 +308,6 @@ bool FAvfMediaPlayer::SetRate( float Rate )
 /* FTickerBase interface
  *****************************************************************************/
 
-bool FAvfMediaPlayer::ShouldAdvanceFrames() const
-{
-    return IsReady() && IsPlaying();
-}
-
 bool FAvfMediaPlayer::Tick( float DeltaTime )
 {
     if( ShouldAdvanceFrames() )
@@ -333,3 +328,12 @@ bool FAvfMediaPlayer::Tick( float DeltaTime )
     
     return true;
 }
+
+
+/* FAvfMediaPlayer implementation
+ *****************************************************************************/
+
+bool FAvfMediaPlayer::ShouldAdvanceFrames() const
+{
+    return IsReady() && IsPlaying();
+}
diff --git a/Engine/Plugins/Media/AvfMedia/Source/AvfMedia/Private/Player/AvfMediaPlayer.h b/Engine/Plugins/Media/AvfMedia/Source/AvfMedia/Private/Player/AvfMediaPlayer.h
--- a/Engine/Plugins/Media/AvfMedia/Source/AvfMedia/Private/Player/AvfMediaPlayer.h
+++ b/Engine/Plugins/Media/AvfMedia/Source/AvfMedia/Private/Player/AvfMediaPlayer.h
@@ -2,6 +2,9 @@
 
 #pragma once
 
+#include "IMediaInfo.h"
+#include "IMediaPlayer.h"
+
 
 /**
  * Implements a media player using the Windows Media Foundation framework.
@@ -67,12 +70,13 @@ public:
 
     virtual bool Tick( float DeltaTime ) override;
 
+protected:
 
-private:
-
-    // FAVFMediaPlayer interface
-
-    /** Should the media player try and progress */
+    /**
+	 * Whether the media player should advance to the next frame.
+	 *
+	 * @return true if the player should advance, false otherwise.
+	 */
     bool ShouldAdvanceFrames() const;
 
 private:
diff --git a/Engine/Plugins/Media/AvfMedia/Source/AvfMedia/Private/Tracks/AvfMediaTrack.h b/Engine/Plugins/Media/AvfMedia/Source/AvfMedia/Private/Tracks/AvfMediaTrack.h
--- a/Engine/Plugins/Media/AvfMedia/Source/AvfMedia/Private/Tracks/AvfMediaTrack.h
+++ b/Engine/Plugins/Media/AvfMedia/Source/AvfMedia/Private/Tracks/AvfMediaTrack.h
@@ -11,8 +11,11 @@ class FAvfMediaTrack
 {
 public:
 
-    FAvfMediaTrack(){}
-    ~FAvfMediaTrack(){}
+	/** Default constructor. */
+    FAvfMediaTrack() { }
+
+	/** Destructor. */
+    ~FAvfMediaTrack() { }
 
 public:
 
