Fix for crash on mac media player when determing player item status.
- no longer check it directly, use delegates to detect changes.

UE-4969

diff --git a/Engine/Plugins/Media/AvfMedia/Source/AvfMedia/Private/Player/AvfMediaPlayer.cpp b/Engine/Plugins/Media/AvfMedia/Source/AvfMedia/Private/Player/AvfMediaPlayer.cpp
--- a/Engine/Plugins/Media/AvfMedia/Source/AvfMedia/Private/Player/AvfMediaPlayer.cpp
+++ b/Engine/Plugins/Media/AvfMedia/Source/AvfMedia/Private/Player/AvfMediaPlayer.cpp
@@ -7,6 +7,47 @@
 #endif
 
 
+/* FMediaHelper implementation
+ *****************************************************************************/
+@implementation FMediaHelper
+@synthesize bIsPlayerItemReady;
+@synthesize MediaPlayer;
+
+-(FMediaHelper*) initWithMediaPlayer:(AVPlayer*)InPlayer
+{
+	MediaPlayer = InPlayer;
+	bIsPlayerItemReady = false;
+	
+	self = [super init];
+	return self;
+}
+
+
+/** Listener for changes in our media classes properties. */
+- (void) observeValueForKeyPath:(NSString*)keyPath
+					ofObject:	(id)object
+					change:		(NSDictionary*)change
+					context:	(void*)context
+{
+	if( [keyPath isEqualToString:@"status"] )
+	{
+		if( object == [MediaPlayer currentItem] )
+		{
+			bIsPlayerItemReady = ([MediaPlayer currentItem].status == AVPlayerItemStatusReadyToPlay);
+		}
+	}
+}
+
+
+- (void)dealloc
+{
+	[MediaPlayer release];
+	[super dealloc];
+}
+
+@end
+
+
 /* FAvfMediaPlayer structors
  *****************************************************************************/
 
@@ -70,11 +111,18 @@ void FAvfMediaPlayer::Close()
 {
     CurrentTime = 0;
 	MediaUrl = FString();
-    
+	
     if( PlayerItem != nil )
-    {
+	{
+		[PlayerItem removeObserver:MediaHelper forKeyPath:@"status"];
         PlayerItem = nil;
     }
+	
+	if( MediaHelper )
+	{
+		[MediaHelper release];
+		MediaHelper = nil;
+	}
 
     if( MediaPlayer != nil )
     {
@@ -150,8 +198,9 @@ bool FAvfMediaPlayer::IsReady() const
 {
     // To be ready, we need the AVPlayer setup
     bool bIsReady = (MediaPlayer != nil) && ([MediaPlayer status] == AVPlayerStatusReadyToPlay);
-    // A player item setup,
-    bIsReady &= (PlayerItem != nil) && ([PlayerItem status] == AVPlayerItemStatusReadyToPlay);
+
+	// A player item setup and ready to stream,
+    bIsReady &= ((MediaHelper != nil) && ([MediaHelper bIsPlayerItemReady]));
     
     // and all tracks to be setup and ready
     for( const IMediaTrackRef& Track : Tracks )
@@ -193,12 +242,15 @@ bool FAvfMediaPlayer::Open( const FString& Url )
         MediaUrl = FPaths::GetCleanFilename(Url);
         
         MediaPlayer = [[AVPlayer alloc] init];
-        
-        if( MediaPlayer )
+		if( MediaPlayer )
         {
+			MediaHelper = [[FMediaHelper alloc] initWithMediaPlayer: MediaPlayer];
+			check( MediaHelper != nil );
+			
             PlayerItem = [AVPlayerItem playerItemWithURL: nsMediaUrl];
             if( PlayerItem != nil )
-            {
+			{
+				[PlayerItem addObserver:MediaHelper forKeyPath:@"status" options:0 context:nil];
                 Duration = FTimespan::FromSeconds( CMTimeGetSeconds( PlayerItem.asset.duration ) );
 
                 [[PlayerItem asset] loadValuesAsynchronouslyForKeys: @[@"tracks"] completionHandler:^
diff --git a/Engine/Plugins/Media/AvfMedia/Source/AvfMedia/Private/Player/AvfMediaPlayer.h b/Engine/Plugins/Media/AvfMedia/Source/AvfMedia/Private/Player/AvfMediaPlayer.h
--- a/Engine/Plugins/Media/AvfMedia/Source/AvfMedia/Private/Player/AvfMediaPlayer.h
+++ b/Engine/Plugins/Media/AvfMedia/Source/AvfMedia/Private/Player/AvfMediaPlayer.h
@@ -6,6 +6,26 @@
 #include "IMediaPlayer.h"
 
 
+/**
+ * Cocoa class that can help us with reading player item information.
+ */
+@interface FMediaHelper : NSObject
+{
+};
+/** We should only initiate a helper with a media player */
+-(FMediaHelper*) initWithMediaPlayer:(AVPlayer*)InPlayer;
+/** Destructor */
+-(void)dealloc;
+
+/** Reference to the media player which will be responsible for this media session */
+@property(retain) AVPlayer* MediaPlayer;
+
+/** Flag to dictate whether the media players current item is ready to play */
+@property bool bIsPlayerItemReady;
+
+@end
+
+
 /**
  * Implements a media player using the Windows Media Foundation framework.
  */
@@ -111,4 +131,8 @@ private:
 
 	/** Holds an event delegate that is invoked when media has been loaded. */
 	FOnMediaOpened OpenedEvent;
+	
+private:
+	/** Cocoa helper object we can use to keep track of ns property changes in our media items */
+	FMediaHelper* MediaHelper;
 };
