AvfMedia: Fixed casing of included file name

GitHub PR #634  by rober-wallis: https://github.com/EpicGames/UnrealEngine/pull/634/files

diff --git a/Engine/Plugins/Media/AvfMedia/Source/AvfMedia/Private/AvfMediaModule.cpp b/Engine/Plugins/Media/AvfMedia/Source/AvfMedia/Private/AvfMediaModule.cpp
--- a/Engine/Plugins/Media/AvfMedia/Source/AvfMedia/Private/AvfMediaModule.cpp
+++ b/Engine/Plugins/Media/AvfMedia/Source/AvfMedia/Private/AvfMediaModule.cpp
@@ -1,6 +1,6 @@
 // Copyright 1998-2015 Epic Games, Inc. All Rights Reserved.
 
-#include "AVFMediaPrivatePCH.h"
+#include "AvfMediaPrivatePCH.h"
 #include "IMediaModule.h"
 #include "IMediaPlayerFactory.h"
 #include "ModuleInterface.h"
diff --git a/Engine/Plugins/Media/AvfMedia/Source/AvfMedia/Private/Tracks/AvfMediaTrack.cpp b/Engine/Plugins/Media/AvfMedia/Source/AvfMedia/Private/Tracks/AvfMediaTrack.cpp
--- a/Engine/Plugins/Media/AvfMedia/Source/AvfMedia/Private/Tracks/AvfMediaTrack.cpp
+++ b/Engine/Plugins/Media/AvfMedia/Source/AvfMedia/Private/Tracks/AvfMediaTrack.cpp
@@ -1,6 +1,6 @@
 // Copyright 1998-2015 Epic Games, Inc. All Rights Reserved.
 
-#include "AVFMediaPrivatePCH.h"
+#include "AvfMediaPrivatePCH.h"
 
 
 #define LOCTEXT_NAMESPACE "FAvfMediaTrack"
