UE-6526 - mac media controls improved
- Added seeking.

diff --git a/Engine/Plugins/Media/AvfMedia/Source/AvfMedia/Private/Player/AvfMediaPlayer.cpp b/Engine/Plugins/Media/AvfMedia/Source/AvfMedia/Private/Player/AvfMediaPlayer.cpp
--- a/Engine/Plugins/Media/AvfMedia/Source/AvfMedia/Private/Player/AvfMediaPlayer.cpp
+++ b/Engine/Plugins/Media/AvfMedia/Source/AvfMedia/Private/Player/AvfMediaPlayer.cpp
@@ -100,7 +100,7 @@ bool FAvfMediaPlayer::SupportsScrubbing() const
 
 bool FAvfMediaPlayer::SupportsSeeking() const
 {
-	return false;
+	return true;
 }
 
 
@@ -326,20 +326,23 @@ bool FAvfMediaPlayer::Open( const TSharedRef<TArray<uint8>>& Buffer, const FStri
 bool FAvfMediaPlayer::Seek( const FTimespan& Time )
 {
     // We need to find a suitable way to seek using avplayer and AVAssetReader,
-/*
     CurrentTime = Time;
-
-    [[MediaPlayer currentItem] seekToTime:CMTimeMakeWithSeconds(CurrentTime.GetSeconds(), 1000)];
+	
+	double TotalSeconds = CurrentTime.GetTotalSeconds();
+	CMTime CurrentTimeInSeconds = CMTimeMake(TotalSeconds, 1);
+	
+	[MediaPlayer seekToTime:CurrentTimeInSeconds];
+	[[MediaPlayer currentItem] seekToTime:CurrentTimeInSeconds];
 
     for( IMediaTrackRef& Track : Tracks )
     {
         FAvfMediaVideoTrack* AVFTrack = (FAvfMediaVideoTrack*)&Track.Get();
         if( AVFTrack != nil )
         {
-            AVFTrack->SeekToTime( [[MediaPlayer currentItem] currentTime] );
+            AVFTrack->SeekToTime(CurrentTimeInSeconds);
         }
     }
-*/
+
     return false;
 }
 
diff --git a/Engine/Plugins/Media/AvfMedia/Source/AvfMedia/Private/Tracks/AvfMediaVideoTrack.cpp b/Engine/Plugins/Media/AvfMedia/Source/AvfMedia/Private/Tracks/AvfMediaVideoTrack.cpp
--- a/Engine/Plugins/Media/AvfMedia/Source/AvfMedia/Private/Tracks/AvfMediaVideoTrack.cpp
+++ b/Engine/Plugins/Media/AvfMedia/Source/AvfMedia/Private/Tracks/AvfMediaVideoTrack.cpp
@@ -14,7 +14,8 @@ FAvfMediaVideoTrack::FAvfMediaVideoTrack( AVAssetTrack* InVideoTrack )
     LatestSamples = nil;
     
     NSError* nsError = nil;
-    AVReader = [[AVAssetReader alloc] initWithAsset: [InVideoTrack asset] error:&nsError];
+	VideoTrack = InVideoTrack;
+    AVReader = [[AVAssetReader alloc] initWithAsset: [VideoTrack asset] error:&nsError];
     if( nsError != nil )
     {
         FString ErrorStr( [nsError localizedDescription] );
@@ -26,7 +27,7 @@ FAvfMediaVideoTrack::FAvfMediaVideoTrack( AVAssetTrack* InVideoTrack )
         NSMutableDictionary* OutputSettings = [NSMutableDictionary dictionary];
         [OutputSettings setObject: [NSNumber numberWithInt:kCVPixelFormatType_32BGRA] forKey:(NSString*)kCVPixelBufferPixelFormatTypeKey];
         
-        AVVideoOutput = [[AVAssetReaderTrackOutput alloc] initWithTrack:InVideoTrack outputSettings:OutputSettings];
+        AVVideoOutput = [[AVAssetReaderTrackOutput alloc] initWithTrack:VideoTrack outputSettings:OutputSettings];
         AVVideoOutput.alwaysCopiesSampleData = NO;
     
         FrameRate = 1.0f / [[AVVideoOutput track] nominalFrameRate];
@@ -60,13 +61,35 @@ FAvfMediaVideoTrack::~FAvfMediaVideoTrack()
 bool FAvfMediaVideoTrack::SeekToTime( const CMTime& SeekTime )
 {
     bool bSeekComplete = false;
-//    if( AVReader != nil )
-//    {
-//        CMTimeRange TimeRange = CMTimeRangeMake(SeekTime, kCMTimePositiveInfinity);
-//        AVReader.timeRange = TimeRange;
-        
-//        bSeekComplete = true;
-//    }
+    if( AVReader != nil )
+    {
+		ResetAssetReader();
+		
+		NSError* nsError = nil;
+		AVReader = [[AVAssetReader alloc] initWithAsset: [VideoTrack asset] error:&nsError];
+		if( nsError != nil )
+		{
+			FString ErrorStr( [nsError localizedDescription] );
+			UE_LOG(LogAvfMedia, Error, TEXT("Failed to create asset reader: %s"), *ErrorStr);
+		}
+		
+        CMTimeRange TimeRange = CMTimeRangeMake(SeekTime, [[AVReader asset] duration]);
+        AVReader.timeRange = TimeRange;
+
+		NSMutableDictionary* OutputSettings = [NSMutableDictionary dictionary];
+		[OutputSettings setObject: [NSNumber numberWithInt:kCVPixelFormatType_32BGRA] forKey:(NSString*)kCVPixelBufferPixelFormatTypeKey];
+		
+		AVVideoOutput = [[AVAssetReaderTrackOutput alloc] initWithTrack:VideoTrack outputSettings:OutputSettings];
+		AVVideoOutput.alwaysCopiesSampleData = NO;
+		
+		[AVReader addOutput:AVVideoOutput];
+		bVideoTracksLoaded = [AVReader startReading];
+		
+		ReadFrameAtTime( SeekTime, false );
+		
+		bSeekComplete = AVReader.status == AVAssetReaderStatusReading;
+    }
+	
     return bSeekComplete;
 }
 
diff --git a/Engine/Plugins/Media/AvfMedia/Source/AvfMedia/Private/Tracks/AvfMediaVideoTrack.h b/Engine/Plugins/Media/AvfMedia/Source/AvfMedia/Private/Tracks/AvfMediaVideoTrack.h
--- a/Engine/Plugins/Media/AvfMedia/Source/AvfMedia/Private/Tracks/AvfMediaVideoTrack.h
+++ b/Engine/Plugins/Media/AvfMedia/Source/AvfMedia/Private/Tracks/AvfMediaVideoTrack.h
@@ -108,6 +108,9 @@ private:
 
 private:
 
+	/** Reference to the av asset, itself */
+	AVAssetTrack* VideoTrack;
+
     /** Access to the Avfoundations track information. */
     AVAssetReader* AVReader;
 
