Added missing variable initialization to FAvfMediaPlayer constructor

diff --git a/Engine/Plugins/Media/AvfMedia/Source/AvfMedia/Private/Player/AvfMediaPlayer.cpp b/Engine/Plugins/Media/AvfMedia/Source/AvfMedia/Private/Player/AvfMediaPlayer.cpp
--- a/Engine/Plugins/Media/AvfMedia/Source/AvfMedia/Private/Player/AvfMediaPlayer.cpp
+++ b/Engine/Plugins/Media/AvfMedia/Source/AvfMedia/Private/Player/AvfMediaPlayer.cpp
@@ -59,6 +59,7 @@ FAvfMediaPlayer::FAvfMediaPlayer()
     
     PlayerItem = nil;
     MediaPlayer = nil;
+	MediaHelper = nil;
 }
 
 
