Media: AVF: Fixed player not in playing state during fast forward/rewind

diff --git a/Engine/Plugins/Media/AvfMedia/Source/AvfMedia/Private/Player/AvfMediaPlayer.cpp b/Engine/Plugins/Media/AvfMedia/Source/AvfMedia/Private/Player/AvfMediaPlayer.cpp
--- a/Engine/Plugins/Media/AvfMedia/Source/AvfMedia/Private/Player/AvfMediaPlayer.cpp
+++ b/Engine/Plugins/Media/AvfMedia/Source/AvfMedia/Private/Player/AvfMediaPlayer.cpp
@@ -191,7 +191,7 @@ bool FAvfMediaPlayer::IsPaused() const
 
 bool FAvfMediaPlayer::IsPlaying() const
 {
-    return (MediaPlayer != nil) && (1.0f == MediaPlayer.rate) && (Tracks.Num() > 0) && (CurrentTime <= Duration);
+	return (MediaPlayer != nil) && !FMath::IsNearlyZero([MediaPlayer rate]) && (Tracks.Num() > 0) && (CurrentTime <= Duration);
 }
 
 
