AvfMedia: Fixing build

diff --git a/Engine/Plugins/Media/AvfMedia/Source/AvfMedia/Private/AvfMediaPrivatePCH.h b/Engine/Plugins/Media/AvfMedia/Source/AvfMedia/Private/AvfMediaPrivatePCH.h
--- a/Engine/Plugins/Media/AvfMedia/Source/AvfMedia/Private/AvfMediaPrivatePCH.h
+++ b/Engine/Plugins/Media/AvfMedia/Source/AvfMedia/Private/AvfMediaPrivatePCH.h
@@ -8,7 +8,7 @@
 
 #include "Core.h"
 #include "IMediaSink.h"
-#include "IMediaTrackInfo.h"
+#include "IMediaStream.h"
 #include "IMediaAudioTrack.h"
 #include "IMediaCaptionTrack.h"
 #include "IMediaVideoTrack.h"
