AvfMedia: Fixing build

diff --git a/Engine/Plugins/Media/AvfMedia/Source/AvfMedia/Private/Player/AvfMediaPlayer.cpp b/Engine/Plugins/Media/AvfMedia/Source/AvfMedia/Private/Player/AvfMediaPlayer.cpp
--- a/Engine/Plugins/Media/AvfMedia/Source/AvfMedia/Private/Player/AvfMediaPlayer.cpp
+++ b/Engine/Plugins/Media/AvfMedia/Source/AvfMedia/Private/Player/AvfMediaPlayer.cpp
@@ -233,7 +233,7 @@ bool FAvfMediaPlayer::IsReady() const
     // and all tracks to be setup and ready
     for( const IMediaAudioTrackRef& AudioTrack : AudioTracks )
     {
-        FAvfMediaVideoTrack* AVFTrack = (FAvfMediaVideoTrack*)&AudioTrack.Get();
+        FAvfMediaTrack* AVFTrack = (FAvfMediaTrack*)&AudioTrack.Get();
 		if ((AVFTrack == nil) || !AVFTrack->IsReady())
 		{
 			return false;
@@ -242,7 +242,7 @@ bool FAvfMediaPlayer::IsReady() const
 
 	for( const IMediaCaptionTrackRef& CaptionTrack : CaptionTracks )
     {
-        FAvfMediaVideoTrack* AVFTrack = (FAvfMediaVideoTrack*)&CaptionTrack.Get();
+        FAvfMediaTrack* AVFTrack = (FAvfMediaTrack*)&CaptionTrack.Get();
 		if ((AVFTrack == nil) || !AVFTrack->IsReady())
 		{
 			return false;
@@ -251,7 +251,7 @@ bool FAvfMediaPlayer::IsReady() const
 
 	for( const IMediaVideoTrackRef& VideoTrack : VideoTracks )
     {
-        FAvfMediaVideoTrack* AVFTrack = (FAvfMediaVideoTrack*)&VideoTrack.Get();
+        FAvfMediaTrack* AVFTrack = (FAvfMediaTrack*)&VideoTrack.Get();
 		if ((AVFTrack == nil) || !AVFTrack->IsReady())
 		{
 			return false;
