AvfMedia: Another blind stab at fixing the build

diff --git a/Engine/Plugins/Media/AvfMedia/Source/AvfMedia/Private/Tracks/AvfMediaAudioTrack.h b/Engine/Plugins/Media/AvfMedia/Source/AvfMedia/Private/Tracks/AvfMediaAudioTrack.h
--- a/Engine/Plugins/Media/AvfMedia/Source/AvfMedia/Private/Tracks/AvfMediaAudioTrack.h
+++ b/Engine/Plugins/Media/AvfMedia/Source/AvfMedia/Private/Tracks/AvfMediaAudioTrack.h
@@ -17,9 +17,18 @@ public:
 	/** Default constructor. */
 	FAvfMediaAudioTrack()
 		: FAvfMediaTrack()
+		, NumChannels(0)
+		, SamplesPerSecond(0)
+	{ }
+
+public:
+
+	// FAvfMediaTrack overrides
+
+	virtual bool IsReady() const override
 	{
-		NumChannels = 0;
-		SamplesPerSecond = 0;
+		// not implemented yet
+		return false;
 	}
 
 public:
diff --git a/Engine/Plugins/Media/AvfMedia/Source/AvfMedia/Private/Tracks/AvfMediaCaptionTrack.h b/Engine/Plugins/Media/AvfMedia/Source/AvfMedia/Private/Tracks/AvfMediaCaptionTrack.h
--- a/Engine/Plugins/Media/AvfMedia/Source/AvfMedia/Private/Tracks/AvfMediaCaptionTrack.h
+++ b/Engine/Plugins/Media/AvfMedia/Source/AvfMedia/Private/Tracks/AvfMediaCaptionTrack.h
@@ -19,6 +19,16 @@ public:
 		: FAvfMediaTrack()
 	{ }
 
+public:
+
+	// FAvfMediaTrack overrides
+
+	virtual bool IsReady() const override
+	{
+		// not implemented yet
+		return false;
+	}
+
 public:
 
 	// IMediaCaptionTrack interface
diff --git a/Engine/Plugins/Media/AvfMedia/Source/AvfMedia/Private/Tracks/AvfMediaTrack.h b/Engine/Plugins/Media/AvfMedia/Source/AvfMedia/Private/Tracks/AvfMediaTrack.h
--- a/Engine/Plugins/Media/AvfMedia/Source/AvfMedia/Private/Tracks/AvfMediaTrack.h
+++ b/Engine/Plugins/Media/AvfMedia/Source/AvfMedia/Private/Tracks/AvfMediaTrack.h
@@ -17,6 +17,15 @@ public:
 	/** Destructor. */
     ~FAvfMediaTrack() { }
 
+public:
+
+	/**
+	 * Check whether the track is ready.
+	 *
+	 * @return true if the track is ready, false otherwise.
+	 */
+	virtual bool IsReady() const = 0;
+
 public:
 
 	// IMediaTrack interface
diff --git a/Engine/Plugins/Media/AvfMedia/Source/AvfMedia/Private/Tracks/AvfMediaVideoTrack.h b/Engine/Plugins/Media/AvfMedia/Source/AvfMedia/Private/Tracks/AvfMediaVideoTrack.h
--- a/Engine/Plugins/Media/AvfMedia/Source/AvfMedia/Private/Tracks/AvfMediaVideoTrack.h
+++ b/Engine/Plugins/Media/AvfMedia/Source/AvfMedia/Private/Tracks/AvfMediaVideoTrack.h
@@ -24,6 +24,12 @@ public:
 	/** Destructor. */
     ~FAvfMediaVideoTrack();
 
+public:
+
+	// FAvfMediaTrack overrides
+
+    virtual bool IsReady() const override;
+
 public:
 
 	// IMediaTrackVideoDetails interface
@@ -48,6 +54,7 @@ public:
 		return *this;
 	}
 
+
 public:
 
 	// IMediaTrack interface
@@ -79,13 +86,6 @@ public:
      */
     bool SeekToTime( const CMTime& SeekTime );
 
-    /**
-     * Is this track ready to begin reading?
-	 *
-	 * @return true if ready, false otherwise.
-     */
-    bool IsReady() const;
-
 private:
 
     /** Reset the AVAssetReader to a different point, which will enable us to regrab frames for seeking and restarting. */
