Fix from Max P, fixing mac.

diff --git a/Engine/Plugins/Media/AvfMedia/Source/AvfMedia/Private/Player/AvfMediaPlayer.cpp b/Engine/Plugins/Media/AvfMedia/Source/AvfMedia/Private/Player/AvfMediaPlayer.cpp
--- a/Engine/Plugins/Media/AvfMedia/Source/AvfMedia/Private/Player/AvfMediaPlayer.cpp
+++ b/Engine/Plugins/Media/AvfMedia/Source/AvfMedia/Private/Player/AvfMediaPlayer.cpp
@@ -233,8 +233,8 @@ bool FAvfMediaPlayer::IsReady() const
     // and all tracks to be setup and ready
     for( const IMediaAudioTrackRef& AudioTrack : AudioTracks )
     {
-        FAvfMediaTrack* AVFTrack = (FAvfMediaTrack*)&AudioTrack.Get();
-		if ((AVFTrack == nil) || !AVFTrack->IsReady())
+        FAvfMediaTrack& AVFTrack = (FAvfMediaTrack&)AudioTrack.Get();
+		if (!AVFTrack.IsReady())
 		{
 			return false;
         }
@@ -242,8 +242,8 @@ bool FAvfMediaPlayer::IsReady() const
 
 	for( const IMediaCaptionTrackRef& CaptionTrack : CaptionTracks )
     {
-        FAvfMediaTrack* AVFTrack = (FAvfMediaTrack*)&CaptionTrack.Get();
-		if ((AVFTrack == nil) || !AVFTrack->IsReady())
+        FAvfMediaTrack& AVFTrack = (FAvfMediaTrack&)CaptionTrack.Get();
+		if (!AVFTrack.IsReady())
 		{
 			return false;
         }
@@ -251,8 +251,8 @@ bool FAvfMediaPlayer::IsReady() const
 
 	for( const IMediaVideoTrackRef& VideoTrack : VideoTracks )
     {
-        FAvfMediaTrack* AVFTrack = (FAvfMediaTrack*)&VideoTrack.Get();
-		if ((AVFTrack == nil) || !AVFTrack->IsReady())
+        FAvfMediaTrack& AVFTrack = (FAvfMediaTrack&)VideoTrack.Get();
+		if (!AVFTrack.IsReady())
 		{
 			return false;
         }
@@ -381,13 +381,10 @@ bool FAvfMediaPlayer::Seek( const FTimespan& Time )
 	[MediaPlayer seekToTime:CurrentTimeInSeconds];
 	[[MediaPlayer currentItem] seekToTime:CurrentTimeInSeconds];
 
-    for( IMediaAudioTrackRef& VideoTrack : VideoTracks )
+    for (IMediaVideoTrackRef& VideoTrack : VideoTracks)
     {
-        FAvfMediaVideoTrack* AVFTrack = (FAvfMediaVideoTrack*)&VideoTrack.Get();
-        if( AVFTrack != nil )
-        {
-            AVFTrack->SeekToTime(CurrentTimeInSeconds);
-        }
+        FAvfMediaVideoTrack& AVFTrack = (FAvfMediaVideoTrack&)VideoTrack.Get();
+        AVFTrack.SeekToTime(CurrentTimeInSeconds);
     }
 
     return false;
@@ -426,12 +423,8 @@ bool FAvfMediaPlayer::Tick( float DeltaTime )
         
     for (IMediaVideoTrackRef& VideoTrack : VideoTracks)
     {
-        FAvfMediaVideoTrack* AVFTrack = (FAvfMediaVideoTrack*)&VideoTrack.Get();
-        
-        if (AVFTrack != nil)
-        {
-            AVFTrack->ReadFrameAtTime([[MediaPlayer currentItem] currentTime]);
-        }
+        FAvfMediaVideoTrack& AVFTrack = (FAvfMediaVideoTrack&)VideoTrack.Get();
+        AVFTrack.ReadFrameAtTime([[MediaPlayer currentItem] currentTime]);
     }
     
     return true;
