Media: Updated comments

diff --git a/Engine/Plugins/Media/AvfMedia/Source/AvfMedia/Private/Tracks/AvfMediaTrack.cpp b/Engine/Plugins/Media/AvfMedia/Source/AvfMedia/Private/Tracks/AvfMediaTrack.cpp
--- a/Engine/Plugins/Media/AvfMedia/Source/AvfMedia/Private/Tracks/AvfMediaTrack.cpp
+++ b/Engine/Plugins/Media/AvfMedia/Source/AvfMedia/Private/Tracks/AvfMediaTrack.cpp
@@ -6,7 +6,7 @@
 #define LOCTEXT_NAMESPACE "FAvfMediaTrack"
 
 
-/* IMediaTrack interface
+/* IMediaStream interface
  *****************************************************************************/
 
 bool FAvfMediaTrack::Disable()
diff --git a/Engine/Plugins/Media/AvfMedia/Source/AvfMedia/Private/Tracks/AvfMediaTrack.h b/Engine/Plugins/Media/AvfMedia/Source/AvfMedia/Private/Tracks/AvfMediaTrack.h
--- a/Engine/Plugins/Media/AvfMedia/Source/AvfMedia/Private/Tracks/AvfMediaTrack.h
+++ b/Engine/Plugins/Media/AvfMedia/Source/AvfMedia/Private/Tracks/AvfMediaTrack.h
@@ -28,7 +28,7 @@ public:
 
 public:
 
-	// IMediaTrack interface
+	// IMediaStream interface
 
 	virtual void AddSink(const IMediaSinkRef& Sink) = 0;
 	virtual bool Disable() override;
diff --git a/Engine/Plugins/Media/AvfMedia/Source/AvfMedia/Private/Tracks/AvfMediaVideoTrack.h b/Engine/Plugins/Media/AvfMedia/Source/AvfMedia/Private/Tracks/AvfMediaVideoTrack.h
--- a/Engine/Plugins/Media/AvfMedia/Source/AvfMedia/Private/Tracks/AvfMediaVideoTrack.h
+++ b/Engine/Plugins/Media/AvfMedia/Source/AvfMedia/Private/Tracks/AvfMediaVideoTrack.h
@@ -32,7 +32,7 @@ public:
 
 public:
 
-	// IMediaTrackVideoDetails interface
+	// IMediaVideoTrack interface
 
 	virtual uint32 GetBitRate() const override
 	{
@@ -57,16 +57,16 @@ public:
 
 public:
 
-	// IMediaTrack interface
+	// IMediaStream interface
 
-    virtual void AddSink( const IMediaSinkRef& Sink ) override
+    virtual void AddSink(const IMediaSinkRef& Sink) override
     {
-        Sinks.Add( Sink );
+        Sinks.Add(Sink);
     }
 
-    virtual void RemoveSink( const IMediaSinkRef& Sink ) override
+    virtual void RemoveSink(const IMediaSinkRef& Sink) override
     {
-        Sinks.Remove( Sink );
+        Sinks.Remove(Sink);
     }
 
 public:
@@ -77,14 +77,14 @@ public:
      * @param AVPlayerTime The time the AVPlayer is currently synced to.
      * @param bInIsInitialFrameRead Flag whether this is the first frame read of the asset.
      */
-    void ReadFrameAtTime( const CMTime& AVPlayerTime, bool bInIsInitialFrameRead = false );
+    void ReadFrameAtTime(const CMTime& AVPlayerTime, bool bInIsInitialFrameRead = false);
 
     /**
      * Set the Av Asset reader to a specified time.
      *
      * @param SeekTime The time the Track Reader should be set to.
      */
-    bool SeekToTime( const CMTime& SeekTime );
+    bool SeekToTime(const CMTime& SeekTime);
 
 private:
 
diff --git a/Engine/Plugins/Media/WmfMedia/Source/WmfMedia/Private/Tracks/WmfMediaTrack.cpp b/Engine/Plugins/Media/WmfMedia/Source/WmfMedia/Private/Tracks/WmfMediaTrack.cpp
--- a/Engine/Plugins/Media/WmfMedia/Source/WmfMedia/Private/Tracks/WmfMediaTrack.cpp
+++ b/Engine/Plugins/Media/WmfMedia/Source/WmfMedia/Private/Tracks/WmfMediaTrack.cpp
@@ -36,7 +36,7 @@ FWmfMediaTrack::FWmfMediaTrack( IMFPresentationDescriptor* InPresentationDescrip
 }
 
 
-/* IMediaTrack interface
+/* IMediaStream interface
  *****************************************************************************/
 
 void FWmfMediaTrack::AddSink( const IMediaSinkRef& Sink )
diff --git a/Engine/Plugins/Media/WmfMedia/Source/WmfMedia/Private/Tracks/WmfMediaTrack.h b/Engine/Plugins/Media/WmfMedia/Source/WmfMedia/Private/Tracks/WmfMediaTrack.h
--- a/Engine/Plugins/Media/WmfMedia/Source/WmfMedia/Private/Tracks/WmfMediaTrack.h
+++ b/Engine/Plugins/Media/WmfMedia/Source/WmfMedia/Private/Tracks/WmfMediaTrack.h
@@ -33,7 +33,7 @@ public:
 
 public:
 
-	// IMediaTrack interface
+	// IMediaStream interface
 
 	virtual void AddSink( const IMediaSinkRef& Sink ) override;
 	virtual bool Disable() override;
diff --git a/Engine/Source/Runtime/Media/Public/IMediaPlayer.h b/Engine/Source/Runtime/Media/Public/IMediaPlayer.h
--- a/Engine/Source/Runtime/Media/Public/IMediaPlayer.h
+++ b/Engine/Source/Runtime/Media/Public/IMediaPlayer.h
@@ -34,7 +34,7 @@ enum class EMediaSeekDirection
 /**
  * Interface for media players.
  *
- * @see IMediaTrack
+ * @see IMediaStream
  */
 class IMediaPlayer
 {
diff --git a/Engine/Source/Runtime/Media/Public/IMediaStream.h b/Engine/Source/Runtime/Media/Public/IMediaStream.h
--- a/Engine/Source/Runtime/Media/Public/IMediaStream.h
+++ b/Engine/Source/Runtime/Media/Public/IMediaStream.h
@@ -33,7 +33,9 @@ public:
 	/**
 	 * Adds the given media sink to this stream.
 	 *
-	 * @param Sink The sink to add.  Uses deferred CPU copies to update the resource.
+	 * This method uses deferred CPU copies to update the resource, which can be slow.
+	 *
+	 * @param Sink The sink to add.
 	 * @see RemoveSink
 	 */
 	virtual void AddSink(const TSharedRef<IMediaSink, ESPMode::ThreadSafe>& Sink) = 0;
@@ -125,8 +127,8 @@ public:
 };
 
 
-/** Type definition for shared pointers to instances of IMediaTrack. */
+/** Type definition for shared pointers to instances of IMediaStream. */
 typedef TSharedPtr<IMediaStream, ESPMode::ThreadSafe> IMediaStreamPtr;
 
-/** Type definition for shared references to instances of IMediaTrack. */
+/** Type definition for shared references to instances of IMediaStream. */
 typedef TSharedRef<IMediaStream, ESPMode::ThreadSafe> IMediaStreamRef;
diff --git a/Engine/Source/Runtime/Media/Public/IMediaVideoTrack.h b/Engine/Source/Runtime/Media/Public/IMediaVideoTrack.h
--- a/Engine/Source/Runtime/Media/Public/IMediaVideoTrack.h
+++ b/Engine/Source/Runtime/Media/Public/IMediaVideoTrack.h
@@ -46,17 +46,19 @@ public:
 
 #if WITH_ENGINE
 	/**
-	* Binds the given TextureResource to the track.  The MediaPlayer/Track implementation can use this to bypass the normal CPU copy deferred sinks
-	* and directly update the resource with the GPU if desired.
-	*
-	* @see RemoveBoundTexture
-	*/
+	 * Bind the given TextureResource to the track.
+	 *
+	 * The MediaPlayer/Track implementation can use this to bypass the normal CPU
+	 * copy deferred sinks and directly update the resource with the GPU if desired.
+	 *
+	 * @see RemoveBoundTexture
+	 */
 	virtual void AddBoundTexture(class FRHITexture* BoundResource) { };
 
 	/**
-	* Removes the given texture from the track binding.  Texture will no longer be updated via BoundTexture path.
-	* @see AddBoundTexture
-	*/
+	 * Remove the given texture from the track binding.  Texture will no longer be updated via BoundTexture path.
+	 * @see AddBoundTexture
+	 */
 	virtual void RemoveBoundTexture(class FRHITexture* BoundResource) { };
 #endif
 
