Media: Added placeholders for fast-path texture binding

This is a new API that allows for bypassing the slow CPU-copied deferred updates of media textures by writing directly into RHI texture resources. MarcusW already implemented it for PS4; we should implement it for the remaining player plug-ins as well. I added placeholders that you can fill in as needed. When all plug-ins are updated, I will switch UMediaTexture to always use the fast path.

diff --git a/Engine/Plugins/Media/AndroidMedia/Source/AndroidMedia/Private/AndroidMediaPlayer.cpp b/Engine/Plugins/Media/AndroidMedia/Source/AndroidMedia/Private/AndroidMediaPlayer.cpp
--- a/Engine/Plugins/Media/AndroidMedia/Source/AndroidMedia/Private/AndroidMediaPlayer.cpp
+++ b/Engine/Plugins/Media/AndroidMedia/Source/AndroidMedia/Private/AndroidMediaPlayer.cpp
@@ -181,6 +181,18 @@ public:
 		return *this;
 	}
 
+#if WITH_ENGINE
+	virtual void BindTexture(class FRHITexture* Texture) override
+	{
+		// @todo android: cbabcock: implement texture binding
+	}
+
+	virtual void UnbindTexture(class FRHITexture* Texture) override
+	{
+		// @todo android: cbabcock: implement texture binding
+	}
+#endif
+
 	virtual void Tick(float DeltaTime)
 	{
 		if (MediaPlayer.MediaState != EMediaState::Error)
diff --git a/Engine/Plugins/Media/AvfMedia/Source/AvfMedia/Private/Tracks/AvfMediaVideoTrack.cpp b/Engine/Plugins/Media/AvfMedia/Source/AvfMedia/Private/Tracks/AvfMediaVideoTrack.cpp
--- a/Engine/Plugins/Media/AvfMedia/Source/AvfMedia/Private/Tracks/AvfMediaVideoTrack.cpp
+++ b/Engine/Plugins/Media/AvfMedia/Source/AvfMedia/Private/Tracks/AvfMediaVideoTrack.cpp
@@ -58,6 +58,32 @@ FAvfMediaVideoTrack::~FAvfMediaVideoTrack()
 /* FAvfMediaVideoTrack interface
  *****************************************************************************/
 
+bool FAvfMediaVideoTrack::IsReady() const
+{
+    return AVReader != nil && AVReader.status == AVAssetReaderStatusReading;
+}
+
+
+/* IMediaVideoTrack interface
+ *****************************************************************************/
+
+#if WITH_ENGINE
+void FAvfMediaVideoTrack::BindTexture(class FRHITexture* Texture)
+{
+	// @todo avf: trepka: implement texture binding
+}
+
+
+void FAvfMediaVideoTrack::UnbindTexture(class FRHITexture* Texture)
+{
+	// @todo avf: trepka: implement texture binding
+}
+#endif
+
+
+/* FAvfMediaVideoTrack implementation
+ *****************************************************************************/
+
 bool FAvfMediaVideoTrack::SeekToTime( const CMTime& SeekTime )
 {
     bool bSeekComplete = false;
@@ -101,12 +127,6 @@ void FAvfMediaVideoTrack::ResetAssetReader()
 }
 
 
-bool FAvfMediaVideoTrack::IsReady() const
-{
-    return AVReader != nil && AVReader.status == AVAssetReaderStatusReading;
-}
-
-
 void FAvfMediaVideoTrack::ReadFrameAtTime( const CMTime& AVPlayerTime, bool bInIsInitialFrameRead )
 {
     check( bVideoTracksLoaded );
diff --git a/Engine/Plugins/Media/AvfMedia/Source/AvfMedia/Private/Tracks/AvfMediaVideoTrack.h b/Engine/Plugins/Media/AvfMedia/Source/AvfMedia/Private/Tracks/AvfMediaVideoTrack.h
--- a/Engine/Plugins/Media/AvfMedia/Source/AvfMedia/Private/Tracks/AvfMediaVideoTrack.h
+++ b/Engine/Plugins/Media/AvfMedia/Source/AvfMedia/Private/Tracks/AvfMediaVideoTrack.h
@@ -54,6 +54,10 @@ public:
 		return *this;
 	}
 
+#if WITH_ENGINE
+	virtual void BindTexture(class FRHITexture* Texture) override;
+	virtual void UnbindTexture(class FRHITexture* Texture) override;
+#endif
 
 public:
 
diff --git a/Engine/Plugins/Media/WmfMedia/Source/WmfMedia/Private/Tracks/WmfMediaVideoTrack.h b/Engine/Plugins/Media/WmfMedia/Source/WmfMedia/Private/Tracks/WmfMediaVideoTrack.h
--- a/Engine/Plugins/Media/WmfMedia/Source/WmfMedia/Private/Tracks/WmfMediaVideoTrack.h
+++ b/Engine/Plugins/Media/WmfMedia/Source/WmfMedia/Private/Tracks/WmfMediaVideoTrack.h
@@ -67,6 +67,11 @@ public:
 		return *this;
 	}
 
+#if WITH_ENGINE
+	virtual void BindTexture(class FRHITexture* Texture) override { }
+	virtual void UnbindTexture(class FRHITexture* Texture) override { }
+#endif
+
 private:
 
 	/** The video's average bit rate. */
diff --git a/Engine/Source/Runtime/Media/Public/IMediaVideoTrack.h b/Engine/Source/Runtime/Media/Public/IMediaVideoTrack.h
--- a/Engine/Source/Runtime/Media/Public/IMediaVideoTrack.h
+++ b/Engine/Source/Runtime/Media/Public/IMediaVideoTrack.h
@@ -46,20 +46,23 @@ public:
 
 #if WITH_ENGINE
 	/**
-	 * Bind the given TextureResource to the track.
+	 * Bind the given texture resource to the track.
 	 *
-	 * The MediaPlayer/Track implementation can use this to bypass the normal CPU
-	 * copy deferred sinks and directly update the resource with the GPU if desired.
+	 * Use this method to bypass the slower CPU-copied deferred media sinks on
+	 * media streams and directly update the resource with the GPU if desired.
 	 *
-	 * @see RemoveBoundTexture
+	 * @param Texture The texture to bind.
+	 * @see UnbindTexture
 	 */
-	virtual void AddBoundTexture(class FRHITexture* BoundResource) { };
+	virtual void BindTexture(class FRHITexture* Texture) = 0;
 
 	/**
-	 * Remove the given texture from the track binding.  Texture will no longer be updated via BoundTexture path.
-	 * @see AddBoundTexture
+	 * Unbind the given texture resource from the track.
+	 *
+	 * @param Texture The texture to unbind.
+	 * @see BindTexture
 	 */
-	virtual void RemoveBoundTexture(class FRHITexture* BoundResource) { };
+	virtual void UnbindTexture(class FRHITexture* Texture) = 0;
 #endif
 
 public:
