Fixed a crash on importing movies on Mac

diff --git a/Engine/Plugins/Media/AvfMedia/Source/AvfMedia/Private/Player/AvfMediaPlayer.cpp b/Engine/Plugins/Media/AvfMedia/Source/AvfMedia/Private/Player/AvfMediaPlayer.cpp
--- a/Engine/Plugins/Media/AvfMedia/Source/AvfMedia/Private/Player/AvfMediaPlayer.cpp
+++ b/Engine/Plugins/Media/AvfMedia/Source/AvfMedia/Private/Player/AvfMediaPlayer.cpp
@@ -251,7 +251,7 @@ bool FAvfMediaPlayer::IsReady() const
 
 	for( const IMediaVideoTrackRef& VideoTrack : VideoTracks )
     {
-        FAvfMediaTrack& AVFTrack = (FAvfMediaTrack&)VideoTrack.Get();
+        FAvfMediaVideoTrack& AVFTrack = (FAvfMediaVideoTrack&)VideoTrack.Get();
 		if (!AVFTrack.IsReady())
 		{
 			return false;
diff --git a/Engine/Plugins/Media/AvfMedia/Source/AvfMedia/Private/Player/AvfMediaPlayer.h b/Engine/Plugins/Media/AvfMedia/Source/AvfMedia/Private/Player/AvfMediaPlayer.h
--- a/Engine/Plugins/Media/AvfMedia/Source/AvfMedia/Private/Player/AvfMediaPlayer.h
+++ b/Engine/Plugins/Media/AvfMedia/Source/AvfMedia/Private/Player/AvfMediaPlayer.h
@@ -27,7 +27,7 @@
 
 
 /**
- * Implements a media player using the Windows Media Foundation framework.
+ * Implements a media player using the AV framework.
  */
 class FAvfMediaPlayer
 	: public IMediaInfo
