Merging CL 2662831 - Improved error handling in Mac media player

diff --git a/Engine/Plugins/Media/AvfMedia/Source/AvfMedia/Private/Player/AvfMediaPlayer.cpp b/Engine/Plugins/Media/AvfMedia/Source/AvfMedia/Private/Player/AvfMediaPlayer.cpp
--- a/Engine/Plugins/Media/AvfMedia/Source/AvfMedia/Private/Player/AvfMediaPlayer.cpp
+++ b/Engine/Plugins/Media/AvfMedia/Source/AvfMedia/Private/Player/AvfMediaPlayer.cpp
@@ -309,24 +309,30 @@ bool FAvfMediaPlayer::Open( const FString& Url )
                         if( AssetVideoTracks.count > 0 )
                         {
                             AVAssetTrack* VideoTrack = [AssetVideoTracks objectAtIndex: 0];
-                            VideoTracks.Add( MakeShareable( new FAvfMediaVideoTrack(VideoTrack) ) );
-
+							FAvfMediaVideoTrack* NewTrack = new FAvfMediaVideoTrack(VideoTrack);
+							if (NewTrack->IsReady())
+							{
+								VideoTracks.Add( MakeShareable( NewTrack ) );
 #if PLATFORM_MAC
-                            GameThreadCall(^{
+								GameThreadCall(^{
 #elif PLATFORM_IOS
-                            // Report back to the game thread whether this succeeded.
-                            [FIOSAsyncTask CreateTaskWithBlock : ^ bool(void){
+									// Report back to the game thread whether this succeeded.
+									[FIOSAsyncTask CreateTaskWithBlock : ^ bool(void){
 #endif
-                                // Displatch on the gamethread that we have opened the video.
-								TracksChangedEvent.Broadcast();
-                                OpenedEvent.Broadcast( CachedUrl );
-
+										// Displatch on the gamethread that we have opened the video.
+										TracksChangedEvent.Broadcast();
+										OpenedEvent.Broadcast( CachedUrl );
 #if PLATFORM_IOS
-                                return true;
-                            }];
+										return true;
+									}];
 #elif PLATFORM_MAC
-                            });
+								});
 #endif
+							}
+							else
+							{
+								delete NewTrack;
+							}
                         }
                     }
                     else
@@ -424,7 +430,10 @@ bool FAvfMediaPlayer::Tick( float DeltaTime )
     for (IMediaVideoTrackRef& VideoTrack : VideoTracks)
     {
         FAvfMediaVideoTrack& AVFTrack = (FAvfMediaVideoTrack&)VideoTrack.Get();
-        AVFTrack.ReadFrameAtTime([[MediaPlayer currentItem] currentTime]);
+        if (!AVFTrack.ReadFrameAtTime([[MediaPlayer currentItem] currentTime]))
+		{
+			return false;
+		}
     }
     
     return true;
diff --git a/Engine/Plugins/Media/AvfMedia/Source/AvfMedia/Private/Tracks/AvfMediaVideoTrack.cpp b/Engine/Plugins/Media/AvfMedia/Source/AvfMedia/Private/Tracks/AvfMediaVideoTrack.cpp
--- a/Engine/Plugins/Media/AvfMedia/Source/AvfMedia/Private/Tracks/AvfMediaVideoTrack.cpp
+++ b/Engine/Plugins/Media/AvfMedia/Source/AvfMedia/Private/Tracks/AvfMediaVideoTrack.cpp
@@ -10,48 +10,58 @@ FAvfMediaVideoTrack::FAvfMediaVideoTrack( AVAssetTrack* InVideoTrack )
     : FAvfMediaTrack()
     , SyncStatus(Default)
     , BitRate(0)
+	, Height(0)
+	, Width(0)
 {
-    LatestSamples = nil;
-    
-    NSError* nsError = nil;
+	LatestSamples = nil;
+
+	NSError* nsError = nil;
 	VideoTrack = InVideoTrack;
-    AVReader = [[AVAssetReader alloc] initWithAsset: [VideoTrack asset] error:&nsError];
-    if( nsError != nil )
-    {
-        FString ErrorStr( [nsError localizedDescription] );
+	AVReader = [[AVAssetReader alloc] initWithAsset: [VideoTrack asset] error:&nsError];
+	if( nsError != nil )
+	{
+		FString ErrorStr( [nsError localizedDescription] );
 		UE_LOG(LogAvfMedia, Error, TEXT("Failed to create asset reader: %s"), *ErrorStr);
-    }
-    else
-    {
-        // Initialize our video output to match the format of the texture we'll be streaming to.
-        NSMutableDictionary* OutputSettings = [NSMutableDictionary dictionary];
-        [OutputSettings setObject: [NSNumber numberWithInt:kCVPixelFormatType_32BGRA] forKey:(NSString*)kCVPixelBufferPixelFormatTypeKey];
-        
-        AVVideoOutput = [[AVAssetReaderTrackOutput alloc] initWithTrack:VideoTrack outputSettings:OutputSettings];
-        AVVideoOutput.alwaysCopiesSampleData = NO;
-    
-        FrameRate = 1.0f / [[AVVideoOutput track] nominalFrameRate];
-        
-        [AVReader addOutput:AVVideoOutput];
-        bVideoTracksLoaded = [AVReader startReading];
+	}
+	else
+	{
+		// Initialize our video output to match the format of the texture we'll be streaming to.
+		NSMutableDictionary* OutputSettings = [NSMutableDictionary dictionary];
+		[OutputSettings setObject: [NSNumber numberWithInt:kCVPixelFormatType_32BGRA] forKey:(NSString*)kCVPixelBufferPixelFormatTypeKey];
 
-        // The initial read of a frame will allow us to gather information on the video track, before we start playing the samples
-        ReadFrameAtTime( kCMTimeZero, true );
-    }
+		AVVideoOutput = [[AVAssetReaderTrackOutput alloc] initWithTrack:VideoTrack outputSettings:OutputSettings];
+		AVVideoOutput.alwaysCopiesSampleData = NO;
+
+		FrameRate = 1.0f / [[AVVideoOutput track] nominalFrameRate];
+
+		[AVReader addOutput:AVVideoOutput];
+		bVideoTracksLoaded = [AVReader startReading];
+
+		// The initial read of a frame will allow us to gather information on the video track, before we start playing the samples
+		if (!ReadFrameAtTime( kCMTimeZero, true ))
+		{
+			[AVVideoOutput release];
+			AVVideoOutput = nil;
+			[AVReader release];
+			AVReader = nil;
+		}
+	}
 }
 
 
 FAvfMediaVideoTrack::~FAvfMediaVideoTrack()
 {
-    if( AVVideoOutput != nil )
-    {
-        AVVideoOutput = nil;
-    }
-    
-    if( AVReader != nil )
-    {
-        AVReader = nil;
-    }
+	if( AVVideoOutput != nil )
+	{
+		[AVVideoOutput release];
+		AVVideoOutput = nil;
+	}
+
+	if( AVReader != nil )
+	{
+		[AVReader release];
+		AVReader = nil;
+	}
 }
 
 
@@ -127,7 +137,7 @@ void FAvfMediaVideoTrack::ResetAssetReader()
 }
 
 
-void FAvfMediaVideoTrack::ReadFrameAtTime( const CMTime& AVPlayerTime, bool bInIsInitialFrameRead )
+bool FAvfMediaVideoTrack::ReadFrameAtTime( const CMTime& AVPlayerTime, bool bInIsInitialFrameRead )
 {
     check( bVideoTracksLoaded );
     
@@ -149,7 +159,7 @@ void FAvfMediaVideoTrack::ReadFrameAtTime( const CMTime& AVPlayerTime, bool bInI
                 if( LatestSamples == NULL )
                 {
                     // Failed to get next sample buffer.
-                    break;
+					return false;
                 }
             }
         
@@ -228,4 +238,6 @@ void FAvfMediaVideoTrack::ReadFrameAtTime( const CMTime& AVPlayerTime, bool bInI
             SyncStatus = bInIsInitialFrameRead ? Ready : Default;
         }
     }
+
+	return true;
 }
diff --git a/Engine/Plugins/Media/AvfMedia/Source/AvfMedia/Private/Tracks/AvfMediaVideoTrack.h b/Engine/Plugins/Media/AvfMedia/Source/AvfMedia/Private/Tracks/AvfMediaVideoTrack.h
--- a/Engine/Plugins/Media/AvfMedia/Source/AvfMedia/Private/Tracks/AvfMediaVideoTrack.h
+++ b/Engine/Plugins/Media/AvfMedia/Source/AvfMedia/Private/Tracks/AvfMediaVideoTrack.h
@@ -81,7 +81,7 @@ public:
      * @param AVPlayerTime The time the AVPlayer is currently synced to.
      * @param bInIsInitialFrameRead Flag whether this is the first frame read of the asset.
      */
-    void ReadFrameAtTime(const CMTime& AVPlayerTime, bool bInIsInitialFrameRead = false);
+    bool ReadFrameAtTime(const CMTime& AVPlayerTime, bool bInIsInitialFrameRead = false);
 
     /**
      * Set the Av Asset reader to a specified time.
