AVF changes
- Added looping.
- Fixed a few bugs related to the file ending and not updating the state of the player
- improved cleanup of the avfmediaplayer

diff --git a/Engine/Plugins/Media/AvfMedia/Source/AvfMedia/Private/Player/AvfMediaPlayer.cpp b/Engine/Plugins/Media/AvfMedia/Source/AvfMedia/Private/Player/AvfMediaPlayer.cpp
--- a/Engine/Plugins/Media/AvfMedia/Source/AvfMedia/Private/Player/AvfMediaPlayer.cpp
+++ b/Engine/Plugins/Media/AvfMedia/Source/AvfMedia/Private/Player/AvfMediaPlayer.cpp
@@ -60,6 +60,8 @@ FAvfMediaPlayer::FAvfMediaPlayer()
     PlayerItem = nil;
     MediaPlayer = nil;
 	MediaHelper = nil;
+	
+	CurrentRate = 0.0f;
 }
 
 
@@ -116,6 +118,7 @@ void FAvfMediaPlayer::Close()
     if( PlayerItem != nil )
     {
 		[PlayerItem removeObserver:MediaHelper forKeyPath:@"status"];
+		[PlayerItem release];
         PlayerItem = nil;
     }
 
@@ -124,12 +127,7 @@ void FAvfMediaPlayer::Close()
 		[MediaHelper release];
 		MediaHelper = nil;
 	}
-
-    if( MediaPlayer != nil )
-    {
-        MediaPlayer = nil;
-    }
-    
+		
 	AudioTracks.Reset();
 	CaptionTracks.Reset();
 	VideoTracks.Reset();
@@ -195,13 +193,13 @@ const TArray<IMediaVideoTrackRef>& FAvfMediaPlayer::GetVideoTracks() const
 
 bool FAvfMediaPlayer::IsLooping() const 
 {
-    return false;
+    return bLoop;
 }
 
 
 bool FAvfMediaPlayer::IsPaused() const
 {
-    return FMath::IsNearlyZero( [MediaPlayer rate] );
+    return FMath::IsNearlyZero( CurrentRate );
 }
 
 
@@ -212,7 +210,7 @@ bool FAvfMediaPlayer::IsPlaying() const
 		return false;
 	}
 
-	return (MediaPlayer != nil) && !FMath::IsNearlyZero([MediaPlayer rate]) && (CurrentTime <= Duration);
+    return (MediaPlayer != nil) && !FMath::IsNearlyZero(CurrentRate) && !ReachedEnd();
 }
 
 
@@ -347,6 +345,7 @@ bool FAvfMediaPlayer::Open( const FString& Url )
             
                 [[MediaPlayer currentItem] seekToTime:kCMTimeZero];
                 MediaPlayer.rate = 0.0;
+				CurrentTime = FTimespan::Zero();
             
                 bSuccessfullyOpenedMovie = true;
             }
@@ -386,7 +385,7 @@ bool FAvfMediaPlayer::Seek( const FTimespan& Time )
 	
 	[MediaPlayer seekToTime:CurrentTimeInSeconds];
 	[[MediaPlayer currentItem] seekToTime:CurrentTimeInSeconds];
-
+	
     for (IMediaVideoTrackRef& VideoTrack : VideoTracks)
     {
         FAvfMediaVideoTrack& AVFTrack = (FAvfMediaVideoTrack&)VideoTrack.Get();
@@ -397,9 +396,10 @@ bool FAvfMediaPlayer::Seek( const FTimespan& Time )
 }
 
 
-bool FAvfMediaPlayer::SetLooping( bool Looping )
+bool FAvfMediaPlayer::SetLooping( bool bInLooping )
 {
-	return false;
+    bLoop = bInLooping;
+	return bLoop;
 }
 
 
@@ -411,31 +411,63 @@ bool FAvfMediaPlayer::SetRate( float Rate )
 }
 
 
+bool FAvfMediaPlayer::ReachedEnd() const
+{
+	bool bHasReachedEnd = false;
+	for( const IMediaVideoTrackRef& VideoTrack : VideoTracks )
+	{
+		FAvfMediaVideoTrack& AVFTrack = (FAvfMediaVideoTrack&)VideoTrack.Get();
+		if (AVFTrack.ReachedEnd())
+		{
+			bHasReachedEnd = true;
+		}
+	}
+	return bHasReachedEnd;
+}
+
+
+bool FAvfMediaPlayer::ShouldTick() const
+{
+	return IsReady() || ReachedEnd();
+}
+
+
 /* FTickerBase interface
  *****************************************************************************/
 
 bool FAvfMediaPlayer::Tick( float DeltaTime )
 {
-    if (!ShouldAdvanceFrames())
-    {
-		return true;
-	}
-
-	AVPlayerItem* item = [MediaPlayer currentItem];
-	if (item != nil)
+	if (ShouldTick())
 	{
-	    CurrentTime = FTimespan::FromSeconds( CMTimeGetSeconds([item currentTime]) );
-	}
-        
-    for (IMediaVideoTrackRef& VideoTrack : VideoTracks)
-    {
-        FAvfMediaVideoTrack& AVFTrack = (FAvfMediaVideoTrack&)VideoTrack.Get();
-        if (!AVFTrack.ReadFrameAtTime([[MediaPlayer currentItem] currentTime]))
+		if (IsPlaying())
 		{
-			return false;
+			for (IMediaVideoTrackRef& VideoTrack : VideoTracks)
+			{
+				FAvfMediaVideoTrack& AVFTrack = (FAvfMediaVideoTrack&)VideoTrack.Get();
+				if (!AVFTrack.ReadFrameAtTime([[MediaPlayer currentItem] currentTime]))
+				{
+					if(bLoop && ReachedEnd())
+					{
+						Seek(FTimespan(0));
+						SetRate(CurrentRate);
+					}
+					else
+					{
+						if(!ReachedEnd())
+						{
+							UE_LOG(LogAvfMedia, Display, TEXT("Failed to read video track for "), *MediaUrl);
+						}
+						return false;
+					}
+				}
+				else
+				{
+					CurrentTime = FTimespan::FromSeconds( CMTimeGetSeconds([[MediaPlayer currentItem] currentTime]) );
+				}
+			}
 		}
-    }
-    
+	}
+	
     return true;
 }
 
@@ -443,7 +475,3 @@ bool FAvfMediaPlayer::Tick( float DeltaTime )
 /* FAvfMediaPlayer implementation
  *****************************************************************************/
 
-bool FAvfMediaPlayer::ShouldAdvanceFrames() const
-{
-    return IsReady() && IsPlaying();
-}
diff --git a/Engine/Plugins/Media/AvfMedia/Source/AvfMedia/Private/Player/AvfMediaPlayer.h b/Engine/Plugins/Media/AvfMedia/Source/AvfMedia/Private/Player/AvfMediaPlayer.h
--- a/Engine/Plugins/Media/AvfMedia/Source/AvfMedia/Private/Player/AvfMediaPlayer.h
+++ b/Engine/Plugins/Media/AvfMedia/Source/AvfMedia/Private/Player/AvfMediaPlayer.h
@@ -95,17 +95,21 @@ public:
 public:
 
     // FTickerObjectBase interface
-
     virtual bool Tick( float DeltaTime ) override;
-
+	
 protected:
-
-    /**
-	 * Whether the media player should advance to the next frame.
+	
+	/**
+	 * Whether the media player should tick in this frame.
 	 *
-	 * @return true if the player should advance, false otherwise.
+	 * @return true if the player should advance, false otherwise
+	 */
+	bool ShouldTick() const;
+	
+	/**
+	 * Has the video completed it's playthrough
 	 */
-    bool ShouldAdvanceFrames() const;
+	bool ReachedEnd() const;
 
 private:
 
@@ -138,6 +142,9 @@ private:
 
 	/** The available video tracks. */
 	TArray<IMediaVideoTrackRef> VideoTracks;
+    
+    /** Should the video loop to the beginning at completion */
+    bool bLoop;
 
 private:
 
diff --git a/Engine/Plugins/Media/AvfMedia/Source/AvfMedia/Private/Tracks/AvfMediaVideoTrack.cpp b/Engine/Plugins/Media/AvfMedia/Source/AvfMedia/Private/Tracks/AvfMediaVideoTrack.cpp
--- a/Engine/Plugins/Media/AvfMedia/Source/AvfMedia/Private/Tracks/AvfMediaVideoTrack.cpp
+++ b/Engine/Plugins/Media/AvfMedia/Source/AvfMedia/Private/Tracks/AvfMediaVideoTrack.cpp
@@ -140,7 +140,7 @@ void FAvfMediaVideoTrack::ResetAssetReader()
 bool FAvfMediaVideoTrack::ReadFrameAtTime( const CMTime& AVPlayerTime, bool bInIsInitialFrameRead )
 {
     check( bVideoTracksLoaded );
-    
+	
     if( AVReader.status == AVAssetReaderStatusReading )
     {
         double CurrentAVTime = CMTimeGetSeconds( AVPlayerTime );
@@ -241,3 +241,9 @@ bool FAvfMediaVideoTrack::ReadFrameAtTime( const CMTime& AVPlayerTime, bool bInI
 
 	return true;
 }
+
+bool FAvfMediaVideoTrack::ReachedEnd() const
+{
+	AVAssetReaderStatus Status = AVReader.status;
+    return (AVReader != nil) && (Status == AVAssetReaderStatusCompleted);
+}
diff --git a/Engine/Plugins/Media/AvfMedia/Source/AvfMedia/Private/Tracks/AvfMediaVideoTrack.h b/Engine/Plugins/Media/AvfMedia/Source/AvfMedia/Private/Tracks/AvfMediaVideoTrack.h
--- a/Engine/Plugins/Media/AvfMedia/Source/AvfMedia/Private/Tracks/AvfMediaVideoTrack.h
+++ b/Engine/Plugins/Media/AvfMedia/Source/AvfMedia/Private/Tracks/AvfMediaVideoTrack.h
@@ -27,7 +27,6 @@ public:
 public:
 
 	// FAvfMediaTrack overrides
-
     virtual bool IsReady() const override;
 
 public:
@@ -89,6 +88,11 @@ public:
      * @param SeekTime The time the Track Reader should be set to.
      */
     bool SeekToTime(const CMTime& SeekTime);
+    
+    /** 
+     * Has the video track completed it's playthrough
+     */
+    bool ReachedEnd() const;
 
 private:
 
