Copying //UE4/Dev-Sequencer to //UE4/Dev-Main (Source: //UE4/Dev-Sequencer @ 3057646)

==========================
MAJOR FEATURES + CHANGES
==========================

Change 2840895 on 2016/01/23 by [REDACTED_PERSON]

	Sequencer: Moved key proxy handling from section into key area; added support for grouped keys; exposed color properties in context menu.

Change 2937981 on 2016/04/08 by [REDACTED_PERSON]

	Sequencer: Refactored GetKeyHandles to take a time range for filtering keys

Change 3051834 on 2016/07/15 by [REDACTED_PERSON]

	Sequencer: Adjust tooltips for autokeying and key all to be more descriptive.

Change 3053057 on 2016/07/17 by [REDACTED_PERSON]

	Sequencer: String track editor

Change 3053083 on 2016/07/18 by [REDACTED_PERSON]

	Sequence Recorder: Add an option to maximize the viewport when starting recording.

Change 3053084 on 2016/07/18 by [REDACTED_PERSON]

	Added UnmapAction to FUICommandList

	Sequencer uses this to re-hook into the level editor commands without issue when it re-opens.

Change 3053085 on 2016/07/18 by [REDACTED_PERSON]

	Sequence Recorder: Display name of next sequence in sequence recorder

	Also display on-screen during countdown

Change 3053086 on 2016/07/18 by [REDACTED_PERSON]

	Sequence Recorder: Improved recording indicator icon, text & timer

	Also fixed a bug with record/stop button where they would not show/enable correctly if no animations were being recorded for a sequence.

Change 3053089 on 2016/07/18 by [REDACTED_PERSON]

	Sequencer: Added hotkey for recording selected actors

	Alt+R now records selected actor into sequencer.

Change 3055488 on 2016/07/19 by [REDACTED_PERSON]

	Sequencer: Fixed NotifyCameraCut() being erroneously called every frame (in FSequencer), and not being called at all in ULevelSequencePlayer

Change 3056783 on 2016/07/19 by [REDACTED_PERSON]

	PS4Media: Attempting to fix playlist advance on PS4 (UE-33481)

diff --git a/Engine/Build/Android/Java/src/com/epicgames/ue4/MediaPlayer14.java b/Engine/Build/Android/Java/src/com/epicgames/ue4/MediaPlayer14.java
--- a/Engine/Build/Android/Java/src/com/epicgames/ue4/MediaPlayer14.java
+++ b/Engine/Build/Android/Java/src/com/epicgames/ue4/MediaPlayer14.java
@@ -12,7 +12,11 @@ import android.os.Build;
 import java.io.File;
 import java.io.IOException;
 import java.io.RandomAccessFile;
+import android.media.MediaExtractor;
+import android.media.MediaFormat;
 import android.media.MediaPlayer;
+import android.media.MediaPlayer.TrackInfo;
+import java.util.ArrayList;
 import java.util.Random;
 import android.util.Log;
 
@@ -23,11 +27,40 @@ import android.util.Log;
 public class MediaPlayer14
 	extends android.media.MediaPlayer
 {
-//	public static Logger Log = new Logger("UE4");
-
 	private boolean SwizzlePixels = true;
 	private volatile boolean WaitOnBitmapRender = false;
 
+	public class AudioTrackInfo {
+		public int Index;
+		public String MimeType;
+		public String DisplayName;
+		public String Language;
+		public int Channels;
+		public int SampleRate;
+	}
+
+	public class CaptionTrackInfo {
+		public int Index;
+		public String MimeType;
+		public String DisplayName;
+		public String Language;
+	}
+
+	public class VideoTrackInfo {
+		public int Index;
+		public String MimeType;
+		public String DisplayName;
+		public String Language;
+		public int BitRate;
+		public int Width;
+		public int Height;
+		public float FrameRate;
+	}
+
+	private ArrayList<AudioTrackInfo> audioTracks = new ArrayList<AudioTrackInfo>();
+	private ArrayList<VideoTrackInfo> videoTracks = new ArrayList<VideoTrackInfo>();
+
+
 	public MediaPlayer14()
 	{
 		SwizzlePixels = true;
@@ -51,12 +84,118 @@ public class MediaPlayer14
 		WaitOnBitmapRender = false;
 	}
 
+
+	private void updateTrackInfo(MediaExtractor extractor)
+	{
+		if (extractor == null)
+		{
+			return;
+		}
+
+		int numTracks = extractor.getTrackCount();
+		int numAudioTracks = 0;
+		int numVideoTracks = 0;
+		audioTracks.ensureCapacity(numTracks);
+		videoTracks.ensureCapacity(numTracks);
+
+		for (int index=0; index < numTracks; index++)
+		{
+			MediaFormat mediaFormat = extractor.getTrackFormat(index);
+			String mimeType = mediaFormat.getString(MediaFormat.KEY_MIME);
+			if (mimeType.startsWith("audio"))
+			{
+				AudioTrackInfo audioTrack = new AudioTrackInfo();
+				audioTrack.Index = index;
+				audioTrack.MimeType = mimeType;
+				audioTrack.DisplayName = "Video Track " + numAudioTracks;
+				audioTrack.Language = "und";
+				audioTrack.Channels = mediaFormat.getInteger(MediaFormat.KEY_CHANNEL_COUNT);
+				audioTrack.SampleRate = mediaFormat.getInteger(MediaFormat.KEY_SAMPLE_RATE);
+				audioTracks.add(audioTrack);
+				numAudioTracks++;
+			}
+			else if (mimeType.startsWith("video"))
+			{
+				VideoTrackInfo videoTrack = new VideoTrackInfo();
+				videoTrack.Index = index;
+				videoTrack.MimeType = mimeType;
+				videoTrack.DisplayName = "Video Track " + numVideoTracks;
+				videoTrack.Language = "und";
+				videoTrack.BitRate = 0;
+				videoTrack.Width = mediaFormat.getInteger(MediaFormat.KEY_WIDTH);
+				videoTrack.Height = mediaFormat.getInteger(MediaFormat.KEY_HEIGHT);
+				videoTrack.FrameRate = mediaFormat.getInteger(MediaFormat.KEY_FRAME_RATE);
+				videoTracks.add(videoTrack);
+				numVideoTracks++;
+			}
+		}
+	}
+
+	private AudioTrackInfo findAudioTrackIndex(int index)
+	{
+		for (AudioTrackInfo track : audioTracks)
+		{
+			if (track.Index == index)
+			{
+				return track;
+			}
+		}
+		return null;
+	}
+
+	private VideoTrackInfo findVideoTrackIndex(int index)
+	{
+		for (VideoTrackInfo track : videoTracks)
+		{
+			if (track.Index == index)
+			{
+				return track;
+			}
+		}
+		return null;
+	}
+
+	public boolean setDataSourceURL(
+		String UrlPath)
+		throws IOException,
+			java.lang.InterruptedException,
+			java.util.concurrent.ExecutionException
+	{
+		audioTracks.clear();
+		videoTracks.clear();
+
+		try
+		{
+			setDataSource(UrlPath);
+			releaseBitmapRenderer();
+			if (android.os.Build.VERSION.SDK_INT >= 16)
+			{
+				MediaExtractor extractor = new MediaExtractor();
+				if (extractor != null)
+				{
+					extractor.setDataSource(UrlPath);
+					updateTrackInfo(extractor);
+					extractor.release();
+					extractor = null;
+				}
+			}
+		}
+		catch(IOException e)
+		{
+			return false;
+		}
+		return true;
+	}
+
 	public boolean setDataSource(
 		String moviePath, long offset, long size)
 		throws IOException,
 			java.lang.InterruptedException,
 			java.util.concurrent.ExecutionException
 	{
+		audioTracks.clear();
+		videoTracks.clear();
+
 		try
 		{
 			File f = new File(moviePath);
@@ -67,6 +206,18 @@ public class MediaPlayer14
 			RandomAccessFile data = new RandomAccessFile(f, "r");
 			setDataSource(data.getFD(), offset, size);
 			releaseBitmapRenderer();
+
+			if (android.os.Build.VERSION.SDK_INT >= 16)
+			{
+				MediaExtractor extractor = new MediaExtractor();
+				if (extractor != null)
+				{
+					extractor.setDataSource(data.getFD(), offset, size);
+					updateTrackInfo(extractor);
+					extractor.release();
+					extractor = null;
+				}
+			}
 		}
 		catch(IOException e)
 		{
@@ -80,11 +231,26 @@ public class MediaPlayer14
 		throws java.lang.InterruptedException,
 			java.util.concurrent.ExecutionException
 	{
+		audioTracks.clear();
+		videoTracks.clear();
+
 		try
 		{
 			AssetFileDescriptor assetFD = assetManager.openFd(assetPath);
 			setDataSource(assetFD.getFileDescriptor(), offset, size);
 			releaseBitmapRenderer();
+
+			if (android.os.Build.VERSION.SDK_INT >= 16)
+			{
+				MediaExtractor extractor = new MediaExtractor();
+				if (extractor != null)
+				{
+					extractor.setDataSource(assetFD.getFileDescriptor(), offset, size);
+					updateTrackInfo(extractor);
+					extractor.release();
+					extractor = null;
+				}
+			}
 		}
 		catch(IOException e)
 		{
@@ -100,9 +266,12 @@ public class MediaPlayer14
 		WaitOnBitmapRender = true;
 
 		mVideoEnabled = enabled;
-		if (mVideoEnabled && null != mBitmapRenderer.getSurface())
+		if (mVideoEnabled)
 		{
-			setSurface(mBitmapRenderer.getSurface());
+			if (null != mBitmapRenderer && null != mBitmapRenderer.getSurface())
+			{
+				setSurface(mBitmapRenderer.getSurface());
+			}
 		}
 		else
 		{
@@ -124,6 +293,15 @@ public class MediaPlayer14
 		}
 	}
 	
+	public boolean didResolutionChange()
+	{
+		if (null == mBitmapRenderer)
+		{
+			return false;
+		}
+		return mBitmapRenderer.resolutionChanged();
+	}
+
 	public void initBitmapRenderer()
 	{
 		// if not already allocated.
@@ -212,6 +390,7 @@ public class MediaPlayer14
 		private int mBlitFragmentShaderID = -1;
 		private float[] mTransformMatrix = new float[16];
 		private boolean mTriangleVerticesDirty = true;
+		private boolean mTextureSizeChanged = true;
 
 		private int GL_TEXTURE_EXTERNAL_OES = 0x8D65;
 		
@@ -381,10 +560,22 @@ public class MediaPlayer14
 					mTextureWidth = width;
 					mTextureHeight = height;
 					mFrameData = null;
+					mTextureSizeChanged = true;
 				}
 			}
 		}
 
+		public boolean resolutionChanged()
+		{
+			boolean changed;
+			synchronized(this)
+			{
+				changed = mTextureSizeChanged;
+				mTextureSizeChanged = false;
+			}
+			return changed;
+		}
+
 		private static final int FLOAT_SIZE_BYTES = 4;
 		private static final int TRIANGLE_VERTICES_DATA_STRIDE_BYTES = 4 * FLOAT_SIZE_BYTES;
 		private static final int TRIANGLE_VERTICES_DATA_POS_OFFSET = 0;
@@ -505,18 +696,18 @@ public class MediaPlayer14
 				mTriangleVerticesData[11] != VMin ||
 				mTriangleVerticesData[3] != VMax)
 			{
-				//Log.debug("Matrix:");
-				//Log.debug(mTransformMatrix[0] + " " + mTransformMatrix[1] + " " + mTransformMatrix[2] + " " + mTransformMatrix[3]);
-				//Log.debug(mTransformMatrix[4] + " " + mTransformMatrix[5] + " " + mTransformMatrix[6] + " " + mTransformMatrix[7]);
-				//Log.debug(mTransformMatrix[8] + " " + mTransformMatrix[9] + " " + mTransformMatrix[10] + " " + mTransformMatrix[11]);
-				//Log.debug(mTransformMatrix[12] + " " + mTransformMatrix[13] + " " + mTransformMatrix[14] + " " + mTransformMatrix[15]);
+				//GameActivity.Log.debug("Matrix:");
+				//GameActivity.Log.debug(mTransformMatrix[0] + " " + mTransformMatrix[1] + " " + mTransformMatrix[2] + " " + mTransformMatrix[3]);
+				//GameActivity.Log.debug(mTransformMatrix[4] + " " + mTransformMatrix[5] + " " + mTransformMatrix[6] + " " + mTransformMatrix[7]);
+				//GameActivity.Log.debug(mTransformMatrix[8] + " " + mTransformMatrix[9] + " " + mTransformMatrix[10] + " " + mTransformMatrix[11]);
+				//GameActivity.Log.debug(mTransformMatrix[12] + " " + mTransformMatrix[13] + " " + mTransformMatrix[14] + " " + mTransformMatrix[15]);
 				mTriangleVerticesData[ 2] = mTriangleVerticesData[10] = UMin;
 				mTriangleVerticesData[ 6] = mTriangleVerticesData[14] = UMax;
 				mTriangleVerticesData[11] = mTriangleVerticesData[15] = VMin;
 				mTriangleVerticesData[ 3] = mTriangleVerticesData[ 7] = VMax;
 				mTriangleVerticesDirty = true;
-				//Log.debug("U = " + mTriangleVerticesData[2] + ", " + mTriangleVerticesData[6]);
-				//Log.debug("V = " + mTriangleVerticesData[11] + ", " + mTriangleVerticesData[3]);
+				//GameActivity.Log.debug("U = " + mTriangleVerticesData[2] + ", " + mTriangleVerticesData[6]);
+				//GameActivity.Log.debug("V = " + mTriangleVerticesData[11] + ", " + mTriangleVerticesData[3]);
 			}
 
 			return true;
@@ -764,6 +955,7 @@ public class MediaPlayer14
 		setOnVideoSizeChangedListener(new android.media.MediaPlayer.OnVideoSizeChangedListener() {
 			public void onVideoSizeChanged(android.media.MediaPlayer player, int w, int h)
 			{
+//				GameActivity.Log.debug("VIDEO SIZE CHANGED: " + w + " x " + h);
 				mBitmapRenderer.setSize(w,h);
 			}
 			});
@@ -784,4 +976,220 @@ public class MediaPlayer14
 	}
 
 	private BitmapRenderer mBitmapRenderer = null;
+
+	public AudioTrackInfo[] GetAudioTracks()
+	{
+		if (android.os.Build.VERSION.SDK_INT >= 16)
+		{
+			TrackInfo[] trackInfo = getTrackInfo();
+			int CountTracks = 0;
+			for (int Index=0; Index < trackInfo.length; Index++)
+			{
+				if (trackInfo[Index].getTrackType() == TrackInfo.MEDIA_TRACK_TYPE_AUDIO)
+				{
+					CountTracks++;
+				}
+			}
+
+			AudioTrackInfo[] AudioTracks = new AudioTrackInfo[CountTracks];
+			int TrackIndex = 0;
+			for (int Index=0; Index < trackInfo.length; Index++)
+			{
+				if (trackInfo[Index].getTrackType() == TrackInfo.MEDIA_TRACK_TYPE_AUDIO)
+				{
+					AudioTracks[TrackIndex] = new AudioTrackInfo();
+					AudioTracks[TrackIndex].Index = Index;
+					AudioTracks[TrackIndex].DisplayName = "Audio Track " + TrackIndex;
+					AudioTracks[TrackIndex].Language = trackInfo[Index].getLanguage();
+	
+					boolean formatValid = false;
+					if (android.os.Build.VERSION.SDK_INT >= 19)
+					{
+						MediaFormat mediaFormat = trackInfo[Index].getFormat();
+						if (mediaFormat != null)
+						{
+							formatValid = true;
+							AudioTracks[TrackIndex].MimeType = mediaFormat.getString(MediaFormat.KEY_MIME);
+							AudioTracks[TrackIndex].Channels = mediaFormat.getInteger(MediaFormat.KEY_CHANNEL_COUNT);
+							AudioTracks[TrackIndex].SampleRate = mediaFormat.getInteger(MediaFormat.KEY_SAMPLE_RATE);
+						}
+					}
+					if (!formatValid && audioTracks.size() > 0)
+					{
+						AudioTrackInfo extractTrack = findAudioTrackIndex(Index);
+						if (extractTrack != null)
+						{
+							formatValid = true;
+							AudioTracks[TrackIndex].MimeType = extractTrack.MimeType;
+							AudioTracks[TrackIndex].Channels = extractTrack.Channels;
+							AudioTracks[TrackIndex].SampleRate = extractTrack.SampleRate;
+						}
+					}
+					if (!formatValid)
+					{
+						AudioTracks[TrackIndex].MimeType = "audio/unknown";
+						AudioTracks[TrackIndex].Channels = 2;
+						AudioTracks[TrackIndex].SampleRate = 44100;
+					}
+
+//					GameActivity.Log.debug(TrackIndex + " (" + Index + ") Audio: Mime=" + AudioTracks[TrackIndex].MimeType + ", Lang=" + AudioTracks[TrackIndex].Language + ", Channels=" + AudioTracks[TrackIndex].Channels + ", SampleRate=" + AudioTracks[TrackIndex].SampleRate);
+					TrackIndex++;
+				}
+			}
+
+			return AudioTracks;
+		}
+
+		AudioTrackInfo[] AudioTracks = new AudioTrackInfo[1];
+
+		AudioTracks[0].Index = 0;
+		AudioTracks[0].MimeType = "audio/unknown";
+		AudioTracks[0].DisplayName = "Audio Track 0";
+		AudioTracks[0].Language = "und";
+		AudioTracks[0].Channels = 2;
+		AudioTracks[0].SampleRate = 44100;
+
+		return AudioTracks;
+	}
+
+	public CaptionTrackInfo[] GetCaptionTracks()
+	{
+		if (android.os.Build.VERSION.SDK_INT >= 21)
+		{
+			TrackInfo[] trackInfo = getTrackInfo();
+			int CountTracks = 0;
+			for (int Index=0; Index < trackInfo.length; Index++)
+			{
+				if (trackInfo[Index].getTrackType() == 4) // TrackInfo.MEDIA_TRACK_TYPE_SUBTITLE
+				{
+					CountTracks++;
+				}
+			}
+
+			CaptionTrackInfo[] CaptionTracks = new CaptionTrackInfo[CountTracks];
+			int TrackIndex = 0;
+			for (int Index=0; Index < trackInfo.length; Index++)
+			{
+				if (trackInfo[Index].getTrackType() == 4) // TrackInfo.MEDIA_TRACK_TYPE_SUBTITLE
+				{
+					CaptionTracks[TrackIndex] = new CaptionTrackInfo();
+					CaptionTracks[TrackIndex].Index = Index;
+					CaptionTracks[TrackIndex].DisplayName = "Caption Track " + TrackIndex;
+
+					MediaFormat mediaFormat = trackInfo[Index].getFormat();
+					if (mediaFormat != null)
+					{
+						CaptionTracks[TrackIndex].MimeType = mediaFormat.getString(MediaFormat.KEY_MIME);
+						CaptionTracks[TrackIndex].Language = mediaFormat.getString(MediaFormat.KEY_LANGUAGE);
+					}
+					else
+					{
+						CaptionTracks[TrackIndex].MimeType = "caption/unknown";
+						CaptionTracks[TrackIndex].Language = trackInfo[Index].getLanguage();
+					}
+
+//					GameActivity.Log.debug(TrackIndex + " (" + Index + ") Caption: Mime=" + CaptionTracks[TrackIndex].MimeType + ", Lang=" + CaptionTracks[TrackIndex].Language);
+					TrackIndex++;
+				}
+			}
+
+			return CaptionTracks;
+		}
+
+		CaptionTrackInfo[] CaptionTracks = new CaptionTrackInfo[0];
+
+		return CaptionTracks;
+	}
+
+	public VideoTrackInfo[] GetVideoTracks()
+	{
+		int Width = getVideoWidth();
+		int Height = getVideoHeight();
+
+		if (android.os.Build.VERSION.SDK_INT >= 16)
+		{
+			TrackInfo[] trackInfo = getTrackInfo();
+			int CountTracks = 0;
+			for (int Index=0; Index < trackInfo.length; Index++)
+			{
+				if (trackInfo[Index].getTrackType() == TrackInfo.MEDIA_TRACK_TYPE_VIDEO)
+				{
+					CountTracks++;
+				}
+			}
+
+			VideoTrackInfo[] VideoTracks = new VideoTrackInfo[CountTracks];
+			int TrackIndex = 0;
+			for (int Index=0; Index < trackInfo.length; Index++)
+			{
+				if (trackInfo[Index].getTrackType() == TrackInfo.MEDIA_TRACK_TYPE_VIDEO)
+				{
+					VideoTracks[TrackIndex] = new VideoTrackInfo();
+					VideoTracks[TrackIndex].Index = Index;
+					VideoTracks[TrackIndex].DisplayName = "Video Track " + TrackIndex;
+					VideoTracks[TrackIndex].Language = trackInfo[Index].getLanguage();
+					VideoTracks[TrackIndex].BitRate = 0;
+
+					boolean formatValid = false;
+					if (android.os.Build.VERSION.SDK_INT >= 19)
+					{
+						MediaFormat mediaFormat = trackInfo[Index].getFormat();
+						if (mediaFormat != null)
+						{
+							formatValid = true;
+							VideoTracks[TrackIndex].MimeType = mediaFormat.getString(MediaFormat.KEY_MIME);
+							VideoTracks[TrackIndex].Width = Integer.parseInt(mediaFormat.getString(MediaFormat.KEY_WIDTH));
+							VideoTracks[TrackIndex].Height = Integer.parseInt(mediaFormat.getString(MediaFormat.KEY_HEIGHT));
+							VideoTracks[TrackIndex].FrameRate = mediaFormat.getFloat(MediaFormat.KEY_FRAME_RATE);
+						}
+					}
+					if (!formatValid && videoTracks.size() > 0)
+					{
+						VideoTrackInfo extractTrack = findVideoTrackIndex(Index);
+						if (extractTrack != null)
+						{
+							formatValid = true;
+							VideoTracks[TrackIndex].MimeType = extractTrack.MimeType;
+							VideoTracks[TrackIndex].Width = extractTrack.Width;
+							VideoTracks[TrackIndex].Height = extractTrack.Height;
+							VideoTracks[TrackIndex].FrameRate = extractTrack.FrameRate;
+						}
+					}
+					if (!formatValid)
+					{
+						VideoTracks[TrackIndex].MimeType = "video/unknown";
+						VideoTracks[TrackIndex].Width = Width;
+						VideoTracks[TrackIndex].Height = Height;
+						VideoTracks[TrackIndex].FrameRate = 30.0f;
+					}
+
+//					GameActivity.Log.debug(TrackIndex + " (" + Index + ") Video: Mime=" + VideoTracks[TrackIndex].MimeType + ", Lang=" + VideoTracks[TrackIndex].Language + ", Width=" + VideoTracks[TrackIndex].Width + ", Height=" + VideoTracks[TrackIndex].Height + ", FrameRate=" + VideoTracks[TrackIndex].FrameRate);
+					TrackIndex++;
+				}
+			}
+
+			return VideoTracks;
+		}
+
+		if (Width > 0 && Height > 0)
+		{
+			VideoTrackInfo[] VideoTracks = new VideoTrackInfo[1];
+
+			VideoTrackInfo VideoTrack = new VideoTrackInfo();
+			VideoTracks[0].Index = 0;
+			VideoTracks[0].MimeType = "video/unknown";
+			VideoTracks[0].DisplayName = "Video Track 0";
+			VideoTracks[0].Language = "und";
+			VideoTracks[0].BitRate = 0;
+			VideoTracks[0].Width = Width;
+			VideoTracks[0].Height = Height;
+			VideoTracks[0].FrameRate = 30.0f;
+
+			return VideoTracks;
+		}
+
+		VideoTrackInfo[] VideoTracks = new VideoTrackInfo[0];
+
+		return VideoTracks;
+	}
 }
diff --git a/Engine/Plugins/Media/AndroidMedia/AndroidMedia.uplugin b/Engine/Plugins/Media/AndroidMedia/AndroidMedia.uplugin
--- a/Engine/Plugins/Media/AndroidMedia/AndroidMedia.uplugin
+++ b/Engine/Plugins/Media/AndroidMedia/AndroidMedia.uplugin
@@ -1,13 +1,14 @@
 {
 	"FileVersion" : 3,
-	"Version" : 1,
-	"VersionName" : "1.0",
+	"Version" : 2,
+	"VersionName" : "2.0",
 	"FriendlyName" : "Android Media Player",
-	"Description" : "Android media player using the Android Media library.",
+	"Description" : "Implements a media player using the Android Media library.",
 	"Category" : "Media Players",
 	"CreatedBy" : "Epic Games, Inc.",
 	"CreatedByURL" : "http://epicgames.com",
-	"DocsURL" : "",
+	"DocsURL" : "https://forums.unrealengine.com/showthread.php?46879-Media-Framework-Documentation-for-4-5-Preview",
+	"EngineVersion" : "4.13.0",
 	"MarketplaceURL" : "",
 	"SupportURL" : "",
 	"EnabledByDefault" : true,
@@ -20,10 +21,18 @@
 			"Name" : "AndroidMedia",
 			"Type" : "Runtime",
 			"LoadingPhase" : "PreLoadingScreen",
-			"WhitelistPlatforms" :
-			[
-				"Android"
-			]
+			"WhitelistPlatforms" : [ "Android" ]
+		},
+		{
+			"Name" : "AndroidMediaEditor",
+			"Type" : "Editor",
+			"LoadingPhase" : "PostEngineInit"
+		},
+		{
+			"Name" : "AndroidMediaFactory",
+			"Type" : "Runtime",
+			"LoadingPhase" : "PostEngineInit",
+			"WhitelistPlatforms" : [ "Android", "Mac", "Linux", "Win32", "Win64" ]
 		}
 	]
 }
\ No newline at end of file
diff --git a/Engine/Plugins/Media/AndroidMedia/Source/AndroidMedia/AndroidMedia.Build.cs b/Engine/Plugins/Media/AndroidMedia/Source/AndroidMedia/AndroidMedia.Build.cs
--- a/Engine/Plugins/Media/AndroidMedia/Source/AndroidMedia/AndroidMedia.Build.cs
+++ b/Engine/Plugins/Media/AndroidMedia/Source/AndroidMedia/AndroidMedia.Build.cs
@@ -26,6 +26,7 @@ namespace UnrealBuildTool.Rules
 			PrivateIncludePaths.AddRange(
 				new string[] {
 					"AndroidMedia/Private",
+					"AndroidMedia/Private/Player",
 				});
 		}
 	}
diff --git a/Engine/Plugins/Media/AndroidMedia/Source/AndroidMedia/Private/AndroidMediaModule.cpp b/Engine/Plugins/Media/AndroidMedia/Source/AndroidMedia/Private/AndroidMediaModule.cpp
--- a/Engine/Plugins/Media/AndroidMedia/Source/AndroidMedia/Private/AndroidMediaModule.cpp
+++ b/Engine/Plugins/Media/AndroidMedia/Source/AndroidMedia/Private/AndroidMediaModule.cpp
@@ -2,111 +2,52 @@
 
 #include "AndroidMediaPCH.h"
 #include "AndroidMediaPlayer.h"
-#include "IMediaModule.h"
-#include "IMediaPlayerFactory.h"
+#include "IAndroidMediaModule.h"
 
 
+DEFINE_LOG_CATEGORY(LogAndroidMedia);
+
 #define LOCTEXT_NAMESPACE "FAndroidMediaModule"
 
 
 class FAndroidMediaModule
-	: public IModuleInterface
-	, public IMediaPlayerFactory
+	: public IAndroidMediaModule
 {
 public:
 
-	// IModuleInterface interface
-
-	virtual void StartupModule() override
-	{
-		if (IsSupported())
-		{
-			IMediaModule* MediaModule = FModuleManager::LoadModulePtr<IMediaModule>("Media");
-			if (nullptr != MediaModule)
-			{
-				SupportedFileTypes.Add(TEXT("3gpp"), LOCTEXT("Format3gpp", "3GPP Multimedia File"));
-				SupportedFileTypes.Add(TEXT("aac"), LOCTEXT("FormatAac", "MPEG-2 Advanced Audio Coding File"));
-				SupportedFileTypes.Add(TEXT("mp4"), LOCTEXT("FormatMp4", "MPEG-4 Movie"));
-
-				// initialize supported URI schemes
-				SupportedUriSchemes.Add(TEXT("http://"));
-				SupportedUriSchemes.Add(TEXT("httpd://"));
-				SupportedUriSchemes.Add(TEXT("https://"));
-				SupportedUriSchemes.Add(TEXT("mms://"));
-				SupportedUriSchemes.Add(TEXT("rtsp://"));
-				SupportedUriSchemes.Add(TEXT("rtspt://"));
-				SupportedUriSchemes.Add(TEXT("rtspu://"));
-
-				MediaModule->RegisterPlayerFactory(*this);
-			}
-		}
-	}
-
-	virtual void ShutdownModule() override
-	{
-		if (IsSupported())
-		{
-			IMediaModule* MediaModule = FModuleManager::GetModulePtr<IMediaModule>("Media");
-			if (nullptr != MediaModule)
-			{
-				MediaModule->UnregisterPlayerFactory(*this);
-			}
-		}
-	}
-
-	// IMediaPlayerFactory interface
+	//~ IAndroidMediaModule interface
 
 	virtual TSharedPtr<IMediaPlayer> CreatePlayer() override
 	{
-		if (IsSupported())
-		{
-			return MakeShareable(new FAndroidMediaPlayer());
-		}
-		else
+		if (!IsSupported())
 		{
 			return nullptr;
 		}
-	}
 
-	virtual const FMediaFileTypes& GetSupportedFileTypes() const override
-	{
-		return SupportedFileTypes;
+		return MakeShareable(new FAndroidMediaPlayer());
 	}
 
-	virtual bool SupportsUrl(const FString& Url) const override
-	{
-		const FString Extension = FPaths::GetExtension(Url);
-
-		if (!Extension.IsEmpty() && SupportedFileTypes.Contains(Extension))
-		{
-			return true;
-		}
-
-		for (const FString& Scheme : SupportedUriSchemes)
-		{
-			if (Url.StartsWith(Scheme))
-			{
-				return true;
-			}
-		}
-
-		return false;
-	}
+public:
 
-private:
+	//~ IModuleInterface interface
 
-	// The collection of supported media file types.
-	FMediaFileTypes SupportedFileTypes;
+	virtual void StartupModule() override { }
+	virtual void ShutdownModule() override { }
 
-	/** The collection of supported URI schemes. */
-	TArray<FString> SupportedUriSchemes;
+protected:
 
+	/**
+	 * Check whether media is supported on the running device.
+	 *
+	 * @return true if media is supported, false otherwise.
+	 */
 	bool IsSupported()
 	{
-		return FAndroidMisc::GetAndroidBuildVersion() >= 14;
+		return (FAndroidMisc::GetAndroidBuildVersion() >= 14);
 	}
 };
 
+
 IMPLEMENT_MODULE(FAndroidMediaModule, AndroidMedia)
 
 #undef LOCTEXT_NAMESPACE
diff --git a/Engine/Plugins/Media/AndroidMedia/Source/AndroidMedia/Private/AndroidMediaPCH.h b/Engine/Plugins/Media/AndroidMedia/Source/AndroidMedia/Private/AndroidMediaPCH.h
--- a/Engine/Plugins/Media/AndroidMedia/Source/AndroidMedia/Private/AndroidMediaPCH.h
+++ b/Engine/Plugins/Media/AndroidMedia/Source/AndroidMedia/Private/AndroidMediaPCH.h
@@ -2,12 +2,17 @@
 
 #pragma once
 
-#include "Core.h"
-#include "IMediaInfo.h"
-#include "IMediaPlayer.h"
-#include "IMediaSink.h"
-#include "IMediaStream.h"
-#include "IMediaAudioTrack.h"
-#include "IMediaCaptionTrack.h"
-#include "IMediaVideoTrack.h"
-#include "ModuleManager.h"
+#include "Runtime/Core/Public/Core.h"
+#include "Runtime/Core/Public/Android/AndroidApplication.h"
+#include "Runtime/Core/Public/Android/AndroidFile.h"
+#include "Runtime/Core/Public/Android/AndroidJava.h"
+#include "Runtime/Core/Public/Misc/Paths.h"
+#include "Runtime/Core/Public/Modules/ModuleManager.h"
+#include "Runtime/Media/Public/IMediaAudioSink.h"
+#include "Runtime/Media/Public/IMediaStringSink.h"
+#include "Runtime/Media/Public/IMediaTextureSink.h"
+#include "Runtime/RenderCore/Public/RenderingThread.h"
+#include "Runtime/RenderCore/Public/RenderResource.h"
+
+
+DECLARE_LOG_CATEGORY_EXTERN(LogAndroidMedia, Log, All);
diff --git a/Engine/Plugins/Media/AndroidMedia/Source/AndroidMedia/Private/Player/AndroidMediaPlayer.cpp b/Engine/Plugins/Media/AndroidMedia/Source/AndroidMedia/Private/Player/AndroidMediaPlayer.cpp
new file mode 100644
--- /dev/null
+++ b/Engine/Plugins/Media/AndroidMedia/Source/AndroidMedia/Private/Player/AndroidMediaPlayer.cpp
@@ -0,0 +1,345 @@
+// Copyright 1998-2016 Epic Games, Inc. All Rights Reserved.
+
+#include "AndroidMediaPCH.h"
+#include "AndroidJavaMediaPlayer.h"
+#include "AndroidMediaPlayer.h"
+#include "AndroidJavaMediaPlayer.h"
+
+
+#define LOCTEXT_NAMESPACE "FAndroidMediaModule"
+
+
+/* FAndroidMediaPlayer structors
+ *****************************************************************************/
+
+FAndroidMediaPlayer::FAndroidMediaPlayer()
+	: JavaMediaPlayer(nullptr)
+{
+	JavaMediaPlayer = MakeShareable(new FJavaAndroidMediaPlayer(false));
+	State = JavaMediaPlayer.IsValid() ? EMediaState::Closed : EMediaState::Error;
+}
+
+
+FAndroidMediaPlayer::~FAndroidMediaPlayer()
+{
+	Close();
+
+	if (JavaMediaPlayer.IsValid())
+	{
+		JavaMediaPlayer->Reset();
+		JavaMediaPlayer->Release();
+	}
+}
+
+
+/* FTickerObjectBase interface
+ *****************************************************************************/
+
+bool FAndroidMediaPlayer::Tick(float DeltaTime)
+{
+	if (State == EMediaState::Playing)
+	{
+		Tracks.Tick();
+
+		if (!JavaMediaPlayer->IsPlaying())
+		{
+			//FPlatformMisc::LowLevelOutputDebugString(TEXT("STOPPED!!!!!"));
+			State = EMediaState::Stopped;
+			MediaEvent.Broadcast(EMediaEvent::PlaybackEndReached);
+		}
+		else if (Tracks.DidPlaybackLoop())
+		{
+			//FPlatformMisc::LowLevelOutputDebugString(TEXT("LOOPED!!!!!"));
+			MediaEvent.Broadcast(EMediaEvent::PlaybackEndReached);
+		}
+	}
+
+	return true;
+}
+
+
+/* IMediaControls interface
+ *****************************************************************************/
+
+FTimespan FAndroidMediaPlayer::GetDuration() const
+{
+	if (State == EMediaState::Error)
+	{
+		return FTimespan::Zero();
+	}
+
+	return FTimespan::FromMilliseconds(JavaMediaPlayer->GetDuration());
+}
+
+
+float FAndroidMediaPlayer::GetRate() const
+{
+	return (State == EMediaState::Playing) ? 1.0f : 0.0f;
+}
+
+
+EMediaState FAndroidMediaPlayer::GetState() const
+{
+	return State;
+}
+
+
+TRange<float> FAndroidMediaPlayer::GetSupportedRates(EMediaPlaybackDirections Direction, bool Unthinned) const
+{
+	if (Direction == EMediaPlaybackDirections::Reverse)
+	{
+		return TRange<float>::Empty();
+	}
+
+	return TRange<float>(1.0f);
+}
+
+
+FTimespan FAndroidMediaPlayer::GetTime() const
+{
+	if (State == EMediaState::Error)
+	{
+		return FTimespan::Zero();
+	}
+
+	return FTimespan::FromMilliseconds(JavaMediaPlayer->GetCurrentPosition());
+}
+
+
+bool FAndroidMediaPlayer::IsLooping() const
+{
+	return JavaMediaPlayer.IsValid() ? JavaMediaPlayer->IsLooping() : false;
+}
+
+
+bool FAndroidMediaPlayer::Seek(const FTimespan& Time)
+{
+	if ((State == EMediaState::Closed) || (State == EMediaState::Error))
+	{
+		return false;
+	}
+
+	JavaMediaPlayer->SeekTo(Time.GetMilliseconds());
+
+	return true;
+}
+
+
+bool FAndroidMediaPlayer::SetLooping(bool Looping)
+{
+	if (!JavaMediaPlayer.IsValid())
+	{
+		return false;
+	}
+
+	JavaMediaPlayer->SetLooping(Looping);
+
+	return true;
+}
+
+
+bool FAndroidMediaPlayer::SetRate(float Rate)
+{
+	switch (State)
+	{
+	case EMediaState::Playing:
+		if (FMath::IsNearlyZero(Rate))
+		{
+			JavaMediaPlayer->Pause();
+			State = EMediaState::Paused;
+			MediaEvent.Broadcast(EMediaEvent::PlaybackSuspended);
+			
+			return true;
+		}
+		break;
+
+	case EMediaState::Paused:
+	case EMediaState::Stopped:
+		if (Rate > 0.0f)
+		{
+			JavaMediaPlayer->Start();
+			State = EMediaState::Playing;
+			MediaEvent.Broadcast(EMediaEvent::PlaybackResumed);
+			
+			return true;
+		}
+		break;
+	}
+
+	return false;
+}
+
+
+bool FAndroidMediaPlayer::SupportsRate(float Rate, bool Unthinned) const
+{
+	return (Rate == 1.0f);
+}
+
+
+bool FAndroidMediaPlayer::SupportsScrubbing() const
+{
+	return true;
+}
+
+
+bool FAndroidMediaPlayer::SupportsSeeking() const
+{
+	return true;
+}
+
+
+/* IMediaPlayer interface
+ *****************************************************************************/
+
+void FAndroidMediaPlayer::Close()
+{
+	if (State == EMediaState::Closed)
+	{
+		return;
+	}
+
+	Tracks.Reset();
+
+	if (JavaMediaPlayer.IsValid())
+	{
+		JavaMediaPlayer->Stop();
+		JavaMediaPlayer->Reset();
+	}
+
+	MediaUrl = FString();
+	State = EMediaState::Closed;
+
+	// notify listeners
+	MediaEvent.Broadcast(EMediaEvent::TracksChanged);
+	MediaEvent.Broadcast(EMediaEvent::MediaClosed);
+}
+
+
+IMediaControls& FAndroidMediaPlayer::GetControls()
+{
+	return *this;
+}
+
+
+FString FAndroidMediaPlayer::GetInfo() const
+{
+	return TEXT("AndroidMedia media information not implemented yet");
+}
+
+
+IMediaOutput& FAndroidMediaPlayer::GetOutput()
+{
+	return Tracks;
+}
+
+
+FString FAndroidMediaPlayer::GetStats() const
+{
+	return TEXT("AndroidMedia stats information not implemented yet");
+}
+
+
+IMediaTracks& FAndroidMediaPlayer::GetTracks()
+{
+	return Tracks;
+}
+
+
+FString FAndroidMediaPlayer::GetUrl() const
+{
+	return MediaUrl;
+}
+
+
+bool FAndroidMediaPlayer::Open(const FString& Url, const IMediaOptions& Options)
+{
+	if (State == EMediaState::Error)
+	{
+		return false;
+	}
+
+	Close();
+
+	if ((Url.IsEmpty()))
+	{
+		return false;
+	}
+
+	MediaUrl = Url;
+
+	// open the media
+	if (Url.StartsWith(TEXT("file://")))
+	{
+		FString FilePath = Url.RightChop(7);
+		FPaths::NormalizeFilename(FilePath);
+
+		IAndroidPlatformFile& PlatformFile = IAndroidPlatformFile::GetPlatformPhysical();
+
+		// make sure that file exists
+		if (!PlatformFile.FileExists(*FilePath))
+		{
+			UE_LOG(LogAndroidMedia, Warning, TEXT("File doesn't exist %s."), *FilePath);
+
+			return false;
+		}
+
+		// get information about media
+		int64 FileOffset = PlatformFile.FileStartOffset(*FilePath);
+		int64 FileSize = PlatformFile.FileSize(*FilePath);
+		FString FileRootPath = PlatformFile.FileRootPath(*FilePath);
+
+		// play movie as a file or asset
+		if (PlatformFile.IsAsset(*FilePath))
+		{
+			if (!JavaMediaPlayer->SetDataSource(PlatformFile.GetAssetManager(), FileRootPath, FileOffset, FileSize))
+			{
+				UE_LOG(LogAndroidMedia, Warning, TEXT("Failed to set data source for asset %s"), *FilePath);
+				return false;
+			}
+		}
+		else
+		{
+			if (!JavaMediaPlayer->SetDataSource(FileRootPath, FileOffset, FileSize))
+			{
+				UE_LOG(LogAndroidMedia, Warning, TEXT("Failed to set data source for file %s"), *FilePath);
+				return false;
+			}
+		}
+	}
+	else
+	{
+		// open remote media
+		JavaMediaPlayer->SetDataSource(Url);
+	}
+
+	// prepare media
+	MediaUrl = Url;
+
+	if (!JavaMediaPlayer->Prepare())
+	{
+		UE_LOG(LogAndroidMedia, Warning, TEXT("Failed to prepare media source %s"), *Url);
+		return false;
+	}
+
+	State = EMediaState::Stopped;
+
+	Tracks.Initialize(JavaMediaPlayer.ToSharedRef());
+	Tracks.SelectTrack(EMediaTrackType::Audio, 0);
+	Tracks.SelectTrack(EMediaTrackType::Video, 0);
+
+	// notify listeners
+	MediaEvent.Broadcast(EMediaEvent::TracksChanged);
+	MediaEvent.Broadcast(EMediaEvent::MediaOpened);
+
+	return true;
+}
+
+
+bool FAndroidMediaPlayer::Open(const TSharedRef<FArchive, ESPMode::ThreadSafe>& Archive, const FString& OriginalUrl, const IMediaOptions& Options)
+{
+	// @todo AndroidMedia: implement opening media from FArchive
+	return false;
+}
+
+
+#undef LOCTEXT_NAMESPACE
diff --git a/Engine/Plugins/Media/AndroidMedia/Source/AndroidMedia/Private/Player/AndroidMediaPlayer.h b/Engine/Plugins/Media/AndroidMedia/Source/AndroidMedia/Private/Player/AndroidMediaPlayer.h
new file mode 100644
--- /dev/null
+++ b/Engine/Plugins/Media/AndroidMedia/Source/AndroidMedia/Private/Player/AndroidMediaPlayer.h
@@ -0,0 +1,88 @@
+// Copyright 1998-2016 Epic Games, Inc. All Rights Reserved.
+
+#pragma once
+
+#include "AndroidMediaTracks.h"
+#include "IMediaControls.h"
+#include "IMediaPlayer.h"
+
+
+class FJavaAndroidMediaPlayer;
+
+
+/*
+	Implement media playback using the Android MediaPlayer interface.
+*/
+class FAndroidMediaPlayer
+	: public FTickerObjectBase
+	, public IMediaControls
+	, public IMediaPlayer
+{
+public:
+	
+	/** Default constructor. */
+	FAndroidMediaPlayer();
+
+	/** Virtual destructor. */
+	virtual ~FAndroidMediaPlayer();
+
+public:
+
+	//~ FTickerObjectBase interface
+
+	virtual bool Tick(float DeltaTime) override;
+
+public:
+
+	//~ IMediaControls interface
+
+	virtual FTimespan GetDuration() const override;
+	virtual float GetRate() const override;
+	virtual EMediaState GetState() const override;
+	virtual TRange<float> GetSupportedRates(EMediaPlaybackDirections Direction, bool Unthinned) const override;
+	virtual FTimespan GetTime() const override;
+	virtual bool IsLooping() const override;
+	virtual bool Seek(const FTimespan& Time) override;
+	virtual bool SetLooping(bool Looping) override;
+	virtual bool SetRate(float Rate) override;
+	virtual bool SupportsRate(float Rate, bool Unthinned) const override;
+	virtual bool SupportsScrubbing() const override;
+	virtual bool SupportsSeeking() const override;
+
+public:
+
+	//~ IMediaPlayer interface
+
+	virtual void Close() override;
+	virtual IMediaControls& GetControls() override;
+	virtual FString GetInfo() const override;
+	virtual IMediaOutput& GetOutput() override;
+	virtual FString GetStats() const override;
+	virtual IMediaTracks& GetTracks() override;
+	virtual FString GetUrl() const override;
+	virtual bool Open(const FString& Url, const IMediaOptions& Options) override;
+	virtual bool Open(const TSharedRef<FArchive, ESPMode::ThreadSafe>& Archive, const FString& OriginalUrl, const IMediaOptions& Options) override;
+
+	DECLARE_DERIVED_EVENT(FVlcMediaPlayer, IMediaPlayer::FOnMediaEvent, FOnMediaEvent);
+	virtual FOnMediaEvent& OnMediaEvent() override
+	{
+		return MediaEvent;
+	}
+
+private:
+
+	/** The Java side media interface. */
+	TSharedPtr<FJavaAndroidMediaPlayer, ESPMode::ThreadSafe> JavaMediaPlayer;
+
+	/** Holds an event delegate that is invoked when a media event occurred. */
+	FOnMediaEvent MediaEvent;
+
+	/** Currently opened media. */
+	FString MediaUrl;
+
+	/** Our understanding of the state of the Java media player. */
+	EMediaState State;
+
+	/** Track collection. */
+	FAndroidMediaTracks Tracks;
+};
diff --git a/Engine/Plugins/Media/AndroidMedia/Source/AndroidMedia/Private/Player/AndroidMediaTracks.cpp b/Engine/Plugins/Media/AndroidMedia/Source/AndroidMedia/Private/Player/AndroidMediaTracks.cpp
new file mode 100644
--- /dev/null
+++ b/Engine/Plugins/Media/AndroidMedia/Source/AndroidMedia/Private/Player/AndroidMediaTracks.cpp
@@ -0,0 +1,485 @@
+// Copyright 1998-2016 Epic Games, Inc. All Rights Reserved.
+
+#include "AndroidMediaPCH.h"
+#include "AndroidMediaTracks.h"
+#include "AndroidJavaMediaPlayer.h"
+
+
+#define LOCTEXT_NAMESPACE "FAndroidMediaTracks"
+
+
+/* FAndroidMediaTracks structors
+ *****************************************************************************/
+
+FAndroidMediaTracks::FAndroidMediaTracks()
+	: AudioSink(nullptr)
+	, CaptionSink(nullptr)
+	, VideoSink(nullptr)
+	, SelectedAudioTrack(INDEX_NONE)
+	, SelectedCaptionTrack(INDEX_NONE)
+	, SelectedVideoTrack(INDEX_NONE)
+	, LastFramePosition(INDEX_NONE)
+	, PlaybackLooped(false)
+{ }
+
+
+FAndroidMediaTracks::~FAndroidMediaTracks()
+{
+	FlushRenderingCommands();
+	Reset();
+}
+
+
+/* FAndroidMediaTracks interface
+ *****************************************************************************/
+
+void FAndroidMediaTracks::Initialize(TSharedRef<FJavaAndroidMediaPlayer, ESPMode::ThreadSafe> InJavaMediaPlayer)
+{
+	Reset();
+
+	FScopeLock Lock(&CriticalSection);
+
+	JavaMediaPlayer = InJavaMediaPlayer;
+	LastFramePosition = INDEX_NONE;
+	PlaybackLooped = false;
+
+	InJavaMediaPlayer->GetAudioTracks(AudioTracks);
+	InJavaMediaPlayer->GetCaptionTracks(CaptionTracks);
+	InJavaMediaPlayer->GetVideoTracks(VideoTracks);
+}
+
+
+void FAndroidMediaTracks::Reset()
+{
+	FScopeLock Lock(&CriticalSection);
+
+	SelectedAudioTrack = INDEX_NONE;
+	SelectedCaptionTrack = INDEX_NONE;
+	SelectedVideoTrack = INDEX_NONE;
+
+	AudioTracks.Empty();
+	CaptionTracks.Empty();
+	VideoTracks.Empty();
+
+	JavaMediaPlayer.Reset();
+}
+
+
+void FAndroidMediaTracks::Tick()
+{
+	UpdateAudioSink();
+	UpdateCaptionSink();
+
+	if (VideoSink != nullptr)
+	{
+		ENQUEUE_UNIQUE_RENDER_COMMAND_ONEPARAMETER(AndroidMediaTracksUpdateVideoSink,
+			FAndroidMediaTracks*, This, this,
+			{
+				This->UpdateVideoSink();
+			});
+	}
+}
+
+
+bool FAndroidMediaTracks::DidPlaybackLoop()
+{
+	bool Result = PlaybackLooped;
+	if (Result)
+	{
+		PlaybackLooped = false;
+	}
+	return Result;
+}
+
+
+/* IMediaOutput interface
+ *****************************************************************************/
+
+void FAndroidMediaTracks::SetAudioSink(IMediaAudioSink* Sink)
+{
+	FScopeLock Lock(&CriticalSection);
+
+	if (Sink != AudioSink)
+	{
+		if (AudioSink != nullptr)
+		{
+			AudioSink->ShutdownAudioSink();
+		}
+
+		AudioSink = Sink;
+		InitializeAudioSink();
+	}
+}
+
+
+void FAndroidMediaTracks::SetCaptionSink(IMediaStringSink* Sink)
+{
+	// not implemented yet
+}
+
+
+void FAndroidMediaTracks::SetImageSink(IMediaTextureSink* Sink)
+{
+	// not supported
+}
+
+
+void FAndroidMediaTracks::SetVideoSink(IMediaTextureSink* Sink)
+{
+	FScopeLock Lock(&CriticalSection);
+
+	if (Sink != VideoSink)
+	{
+		if (VideoSink != nullptr)
+		{
+			VideoSink->ShutdownTextureSink();
+		}
+
+		VideoSink = Sink;
+		InitializeVideoSink();
+	}
+}
+
+
+/* IMediaTracks interface
+ *****************************************************************************/
+
+uint32 FAndroidMediaTracks::GetAudioTrackChannels(int32 TrackIndex) const
+{
+	return AudioTracks.IsValidIndex(TrackIndex) ? AudioTracks[TrackIndex].Channels : 0;
+}
+
+
+uint32 FAndroidMediaTracks::GetAudioTrackSampleRate(int32 TrackIndex) const
+{
+	return AudioTracks.IsValidIndex(TrackIndex) ? AudioTracks[TrackIndex].SampleRate : 0;
+}
+
+
+int32 FAndroidMediaTracks::GetNumTracks(EMediaTrackType TrackType) const
+{
+	switch (TrackType)
+	{
+	case EMediaTrackType::Audio:
+		return AudioTracks.Num();
+	case EMediaTrackType::Caption:
+		return CaptionTracks.Num();
+	case EMediaTrackType::Video:
+		return VideoTracks.Num();
+	default:
+		return 0;
+	}
+}
+
+
+int32 FAndroidMediaTracks::GetSelectedTrack(EMediaTrackType TrackType) const
+{
+	switch (TrackType)
+	{
+	case EMediaTrackType::Audio:
+		return SelectedAudioTrack;
+
+	case EMediaTrackType::Caption:
+		return SelectedCaptionTrack;
+
+	case EMediaTrackType::Video:
+		return SelectedVideoTrack;
+
+	default:
+		return INDEX_NONE;
+	}
+}
+
+
+FText FAndroidMediaTracks::GetTrackDisplayName(EMediaTrackType TrackType, int32 TrackIndex) const
+{
+	switch (TrackType)
+	{
+	case EMediaTrackType::Audio:
+		if (AudioTracks.IsValidIndex(TrackIndex))
+		{
+			return FText::FromString(AudioTracks[TrackIndex].DisplayName);
+		}
+		break;
+
+	case EMediaTrackType::Caption:
+		if (CaptionTracks.IsValidIndex(TrackIndex))
+		{
+			return FText::FromString(CaptionTracks[TrackIndex].DisplayName);
+		}
+		break;
+
+	case EMediaTrackType::Video:
+		if (VideoTracks.IsValidIndex(TrackIndex))
+		{
+			return FText::FromString(VideoTracks[TrackIndex].DisplayName);
+		}
+
+	default:
+		break;
+	}
+
+	return FText::GetEmpty();
+}
+
+
+FString FAndroidMediaTracks::GetTrackLanguage(EMediaTrackType TrackType, int32 TrackIndex) const
+{
+	switch (TrackType)
+	{
+	case EMediaTrackType::Audio:
+		if (AudioTracks.IsValidIndex(TrackIndex))
+		{
+			return AudioTracks[TrackIndex].Language;
+		}
+		break;
+
+	case EMediaTrackType::Caption:
+		if (CaptionTracks.IsValidIndex(TrackIndex))
+		{
+			return CaptionTracks[TrackIndex].Language;
+		}
+		break;
+
+	case EMediaTrackType::Video:
+		if (VideoTracks.IsValidIndex(TrackIndex))
+		{
+			return VideoTracks[TrackIndex].Language;
+		}
+
+	default:
+		break;
+	}
+
+	return FString();
+}
+
+
+FString FAndroidMediaTracks::GetTrackName(EMediaTrackType TrackType, int32 TrackIndex) const
+{
+	switch (TrackType)
+	{
+	case EMediaTrackType::Audio:
+		if (AudioTracks.IsValidIndex(TrackIndex))
+		{
+			return AudioTracks[TrackIndex].Name;
+		}
+		break;
+
+	case EMediaTrackType::Caption:
+		if (CaptionTracks.IsValidIndex(TrackIndex))
+		{
+			return CaptionTracks[TrackIndex].Name;
+		}
+		break;
+
+	case EMediaTrackType::Video:
+		if (VideoTracks.IsValidIndex(TrackIndex))
+		{
+			return VideoTracks[TrackIndex].Name;
+		}
+
+	default:
+		break;
+	}
+
+	return FString();
+}
+
+
+uint32 FAndroidMediaTracks::GetVideoTrackBitRate(int32 TrackIndex) const
+{
+	return VideoTracks.IsValidIndex(TrackIndex) ? VideoTracks[TrackIndex].BitRate : 0;
+}
+
+
+FIntPoint FAndroidMediaTracks::GetVideoTrackDimensions(int32 TrackIndex) const
+{
+	return VideoTracks.IsValidIndex(TrackIndex)
+		? VideoTracks[TrackIndex].Dimensions
+		: FIntPoint(0, 0);
+}
+
+
+float FAndroidMediaTracks::GetVideoTrackFrameRate(int32 TrackIndex) const
+{
+	return VideoTracks.IsValidIndex(TrackIndex) ? VideoTracks[TrackIndex].FrameRate : 0;
+}
+
+
+bool FAndroidMediaTracks::SelectTrack(EMediaTrackType TrackType, int32 TrackIndex)
+{
+	FScopeLock Lock(&CriticalSection);
+
+	switch (TrackType)
+	{
+	case EMediaTrackType::Audio:
+		if (AudioTracks.IsValidIndex(TrackIndex))
+		{
+			if (!JavaMediaPlayer->SelectTrack(AudioTracks[TrackIndex].Index))
+			{
+				return false;
+			}
+
+			SelectedAudioTrack = TrackIndex;
+			InitializeAudioSink();
+
+			return true;
+		}
+		break;
+
+	case EMediaTrackType::Caption:
+		if (CaptionTracks.IsValidIndex(TrackIndex))
+		{
+			if (!JavaMediaPlayer->SelectTrack(CaptionTracks[TrackIndex].Index))
+			{
+				return false;
+			}
+
+			SelectedCaptionTrack = TrackIndex;
+		}
+		break;
+
+	case EMediaTrackType::Video:
+		if (TrackIndex == INDEX_NONE)
+		{
+			JavaMediaPlayer->SetVideoEnabled(false);
+			SelectedVideoTrack = TrackIndex;
+
+			return true;
+		}
+		
+		if (VideoTracks.IsValidIndex(TrackIndex))
+		{
+			if (!JavaMediaPlayer->SelectTrack(VideoTracks[TrackIndex].Index))
+			{
+				return false;
+			}
+
+			JavaMediaPlayer->SetVideoEnabled(true);
+			SelectedVideoTrack = TrackIndex;
+			InitializeVideoSink();
+
+			return true;
+		}
+		break;
+	}
+
+	return false;
+}
+
+
+/* FAndroidMediaTracks implementation
+ *****************************************************************************/
+
+void FAndroidMediaTracks::InitializeAudioSink()
+{
+	if ((AudioSink == nullptr) || (SelectedAudioTrack == INDEX_NONE))
+	{
+		return;
+	}
+
+	const auto& AudioTrack = AudioTracks[SelectedAudioTrack];
+	AudioSink->InitializeAudioSink(AudioTrack.Channels, AudioTrack.SampleRate);
+}
+
+
+void FAndroidMediaTracks::InitializeCaptionSink()
+{
+	if ((CaptionSink == nullptr) || (SelectedCaptionTrack == INDEX_NONE))
+	{
+		return;
+	}
+
+	const auto& CaptionTrack = CaptionTracks[SelectedCaptionTrack];
+	CaptionSink->InitializeStringSink();
+}
+
+
+void FAndroidMediaTracks::InitializeVideoSink()
+{
+	if ((VideoSink == nullptr) || (SelectedVideoTrack == INDEX_NONE))
+	{
+		return;
+	}
+
+	const auto& VideoTrack = VideoTracks[SelectedVideoTrack];
+	VideoSink->InitializeTextureSink(VideoTrack.Dimensions, EMediaTextureSinkFormat::CharBGRA, EMediaTextureSinkMode::Unbuffered);
+}
+
+
+void FAndroidMediaTracks::UpdateAudioSink()
+{
+	// android currently doesn't have an API for extracting audio
+}
+
+
+void FAndroidMediaTracks::UpdateCaptionSink()
+{
+	// not implemented yet
+}
+
+
+void FAndroidMediaTracks::UpdateVideoSink()
+{
+	FScopeLock Lock(&CriticalSection);
+
+	if (!JavaMediaPlayer.IsValid())
+	{
+		return;
+	}
+
+	// only update if playback position changed
+	int32 CurrentFramePosition = JavaMediaPlayer->GetCurrentPosition();
+
+	if (LastFramePosition == CurrentFramePosition)
+	{
+		return;
+	}
+
+	// deal with resolution changes (usually from streams)
+	if (JavaMediaPlayer->DidResolutionChange())
+	{
+		FIntPoint Dimensions = FIntPoint(JavaMediaPlayer->GetVideoWidth(), JavaMediaPlayer->GetVideoHeight());
+		if (SelectedVideoTrack != INDEX_NONE)
+		{
+			// The video track dimensions need updating
+			VideoTracks[SelectedVideoTrack].Dimensions = Dimensions;
+		}
+		VideoSink->InitializeTextureSink(Dimensions, EMediaTextureSinkFormat::CharBGRA, EMediaTextureSinkMode::Unbuffered);
+	}
+
+	// update sink
+#if WITH_ENGINE
+	FRHITexture* Texture = VideoSink->GetTextureSinkTexture();
+
+	if (Texture != nullptr)
+	{
+		int32 Resource = *reinterpret_cast<int32*>(Texture->GetNativeResource());
+
+		if (JavaMediaPlayer->GetVideoLastFrame(Resource))
+		{
+			if (CurrentFramePosition < LastFramePosition)
+			{
+				PlaybackLooped = true;
+			}
+			LastFramePosition = CurrentFramePosition;
+		}
+	}
+#else
+	void* Buffer = VideoSink->AcquireTextureSinkBuffer();
+	int64 SampleCount = 0;
+
+	if ((Buffer != nullptr) && JavaMediaPlayer->GetVideoLastFrameData(Buffer, SampleCount))
+	{
+		VideoSink->ReleaseTextureSinkBuffer();
+		VideoSink->DisplayTextureSinkBuffer();
+
+		if (CurrentFramePosition < LastFramePosition)
+		{
+			PlaybackLooped = true;
+		}
+		LastFramePosition = CurrentFramePosition;
+	}
+#endif
+}
diff --git a/Engine/Plugins/Media/AndroidMedia/Source/AndroidMedia/Private/Player/AndroidMediaTracks.h b/Engine/Plugins/Media/AndroidMedia/Source/AndroidMedia/Private/Player/AndroidMediaTracks.h
new file mode 100644
--- /dev/null
+++ b/Engine/Plugins/Media/AndroidMedia/Source/AndroidMedia/Private/Player/AndroidMediaTracks.h
@@ -0,0 +1,136 @@
+// Copyright 1998-2016 Epic Games, Inc. All Rights Reserved.
+
+#pragma once
+
+#include "IMediaOutput.h"
+#include "IMediaTracks.h"
+#include "AndroidJavaMediaPlayer.h"
+
+
+
+/**
+ * Implements the track collection for Android based media players.
+ */
+class FAndroidMediaTracks
+	: public IMediaOutput
+	, public IMediaTracks
+{
+
+public:
+
+	/** Default constructor. */
+	FAndroidMediaTracks();
+
+	/** Destructor. */
+	~FAndroidMediaTracks();
+
+public:
+
+	/**
+	 * Initialize this object for the specified VLC media player.
+	 *
+	 * @param InJavaMediaPlayer The java player to use.
+	 */
+	void Initialize(TSharedRef<FJavaAndroidMediaPlayer, ESPMode::ThreadSafe> InJavaMediaPlayer);
+
+	/** Reset this object. */
+	void Reset();
+
+	/** Called each frame. */
+	void Tick();
+
+	/** Called to check if playback looped */
+	bool DidPlaybackLoop();
+
+public:
+
+	//~ IMediaOutput interface
+
+	virtual void SetAudioSink(IMediaAudioSink* Sink) override;
+	virtual void SetCaptionSink(IMediaStringSink* Sink) override;
+	virtual void SetImageSink(IMediaTextureSink* Sink) override;
+	virtual void SetVideoSink(IMediaTextureSink* Sink) override;
+
+public:
+
+	//~ IMediaTracks interface
+
+	virtual uint32 GetAudioTrackChannels(int32 TrackIndex) const override;
+	virtual uint32 GetAudioTrackSampleRate(int32 TrackIndex) const override;
+	virtual int32 GetNumTracks(EMediaTrackType TrackType) const override;
+	virtual int32 GetSelectedTrack(EMediaTrackType TrackType) const override;
+	virtual FText GetTrackDisplayName(EMediaTrackType TrackType, int32 TrackIndex) const override;
+	virtual FString GetTrackLanguage(EMediaTrackType TrackType, int32 TrackIndex) const override;
+	virtual FString GetTrackName(EMediaTrackType TrackType, int32 TrackIndex) const override;
+	virtual uint32 GetVideoTrackBitRate(int32 TrackIndex) const override;
+	virtual FIntPoint GetVideoTrackDimensions(int32 TrackIndex) const override;
+	virtual float GetVideoTrackFrameRate(int32 TrackIndex) const override;
+	virtual bool SelectTrack(EMediaTrackType TrackType, int32 TrackIndex) override;
+
+protected:
+
+	/** Initialize the current audio sink. */
+	void InitializeAudioSink();
+
+	/** Initialize the current caption sink. */
+	void InitializeCaptionSink();
+
+	/** Initialize the current video sink. */
+	void InitializeVideoSink();
+
+	/** Send the latest sample data to the audio sink. */
+	void UpdateAudioSink();
+
+	/** Send the latest sample data to the caption sink. */
+	void UpdateCaptionSink();
+
+	/** Send the latest frame to the video sink. */
+	void UpdateVideoSink();
+
+private:
+
+	/** Index of the selected audio track. */
+	int32 SelectedAudioTrack;
+
+	/** Index of the selected audio track. */
+	int32 SelectedCaptionTrack;
+
+	/** Index of the selected video track. */
+	int32 SelectedVideoTrack;
+
+private:
+
+	/** The audio sink. */
+	IMediaAudioSink* AudioSink;
+
+	/** The caption sink. */
+	IMediaStringSink* CaptionSink;
+
+	/** The video sink. */
+	IMediaTextureSink* VideoSink;
+
+private:
+
+	/** Audio track descriptors. */
+	TArray<FJavaAndroidMediaPlayer::FAudioTrack> AudioTracks;
+
+	/** Caption track descriptors. */
+	TArray<FJavaAndroidMediaPlayer::FCaptionTrack> CaptionTracks;
+
+	/** Video track descriptors. */
+	TArray<FJavaAndroidMediaPlayer::FVideoTrack> VideoTracks;
+
+private:
+
+	/** Synchronizes write access to track arrays, selections & sinks. */
+	FCriticalSection CriticalSection;
+
+	/** The Java side media interface. */
+	TSharedPtr<FJavaAndroidMediaPlayer, ESPMode::ThreadSafe> JavaMediaPlayer;
+
+	/** Position of the last video frame being displayed. */
+	int32 LastFramePosition;
+
+	/** Flag for playback looping detected. */
+	bool PlaybackLooped;
+};
diff --git a/Engine/Plugins/Media/AndroidMedia/Source/AndroidMedia/Public/IAndroidMediaModule.h b/Engine/Plugins/Media/AndroidMedia/Source/AndroidMedia/Public/IAndroidMediaModule.h
new file mode 100644
--- /dev/null
+++ b/Engine/Plugins/Media/AndroidMedia/Source/AndroidMedia/Public/IAndroidMediaModule.h
@@ -0,0 +1,30 @@
+// Copyright 1998-2016 Epic Games, Inc. All Rights Reserved.
+
+#pragma once
+
+#include "ModuleInterface.h"
+
+
+class IMediaPlayer;
+
+
+/**
+ * Interface for the AndroidMedia module.
+ */
+class IAndroidMediaModule
+	: public IModuleInterface
+{
+public:
+
+	/**
+	 * Creates a Android based media player.
+	 *
+	 * @return A new media player, or nullptr if a player couldn't be created.
+	 */
+	virtual TSharedPtr<IMediaPlayer> CreatePlayer() = 0;
+
+public:
+
+	/** Virtual destructor. */
+	virtual ~IAndroidMediaModule() { }
+};
diff --git a/Engine/Plugins/Media/AndroidMedia/Source/AndroidMediaEditor/AndroidMediaEditor.Build.cs b/Engine/Plugins/Media/AndroidMedia/Source/AndroidMediaEditor/AndroidMediaEditor.Build.cs
new file mode 100644
--- /dev/null
+++ b/Engine/Plugins/Media/AndroidMedia/Source/AndroidMediaEditor/AndroidMediaEditor.Build.cs
@@ -0,0 +1,25 @@
+// Copyright 1998-2016 Epic Games, Inc. All Rights Reserved.
+
+namespace UnrealBuildTool.Rules
+{
+	public class AndroidMediaEditor : ModuleRules
+	{
+		public AndroidMediaEditor(TargetInfo Target)
+		{
+			PrivateDependencyModuleNames.AddRange(
+				new string[] {
+					"Core",
+					"CoreUObject",
+					"MediaAssets",
+					"UnrealEd",
+                }
+			);
+
+			PrivateIncludePaths.AddRange(
+				new string[] {
+					"AndroidMediaEditor/Private",
+				}
+			);
+		}
+	}
+}
diff --git a/Engine/Plugins/Media/AndroidMedia/Source/AndroidMediaEditor/Private/AndroidMediaEditorModule.cpp b/Engine/Plugins/Media/AndroidMedia/Source/AndroidMediaEditor/Private/AndroidMediaEditorModule.cpp
new file mode 100644
--- /dev/null
+++ b/Engine/Plugins/Media/AndroidMedia/Source/AndroidMediaEditor/Private/AndroidMediaEditorModule.cpp
@@ -0,0 +1,22 @@
+// Copyright 1998-2016 Epic Games, Inc. All Rights Reserved.
+
+#include "AndroidMediaEditorPCH.h"
+#include "ModuleInterface.h"
+
+
+/**
+ * Implements the AndroidMediaEditor module.
+ */
+class FAndroidMediaEditorModule
+	: public IModuleInterface
+{
+public:
+
+	//~ IModuleInterface interface
+
+	virtual void StartupModule() override { }
+	virtual void ShutdownModule() override { }
+};
+
+
+IMPLEMENT_MODULE(FAndroidMediaEditorModule, AndroidMediaEditor);
diff --git a/Engine/Plugins/Media/AndroidMedia/Source/AndroidMediaEditor/Private/AndroidMediaEditorPCH.h b/Engine/Plugins/Media/AndroidMedia/Source/AndroidMediaEditor/Private/AndroidMediaEditorPCH.h
new file mode 100644
--- /dev/null
+++ b/Engine/Plugins/Media/AndroidMedia/Source/AndroidMediaEditor/Private/AndroidMediaEditorPCH.h
@@ -0,0 +1,6 @@
+// Copyright 1998-2016 Epic Games, Inc. All Rights Reserved.
+
+#pragma once
+
+#include "Runtime/Core/Public/Core.h"
+#include "Runtime/MediaAssets/Public/FileMediaSource.h"
diff --git a/Engine/Plugins/Media/AndroidMedia/Source/AndroidMediaFactory/Private/AndroidMediaFactoryPCH.h b/Engine/Plugins/Media/AndroidMedia/Source/AndroidMediaFactory/Private/AndroidMediaFactoryPCH.h
new file mode 100644
--- /dev/null
+++ b/Engine/Plugins/Media/AndroidMedia/Source/AndroidMediaFactory/Private/AndroidMediaFactoryPCH.h
@@ -0,0 +1,15 @@
+// Copyright 1998-2016 Epic Games, Inc. All Rights Reserved.
+
+#pragma once
+
+#include "Runtime/Core/Public/Core.h"
+#include "Runtime/Core/Public/Modules/ModuleManager.h"
+#include "Runtime/CoreUObject/Public/CoreUObject.h"
+#include "Runtime/Media/Public/IMediaModule.h"
+#include "Runtime/Media/Public/IMediaOptions.h"
+#include "Runtime/MediaAssets/Public/FileMediaSource.h"
+
+#include "../../AndroidMedia/Public/IAndroidMediaModule.h"
+
+
+DECLARE_LOG_CATEGORY_EXTERN(LogAndroidMediaFactory, Log, All);
diff --git a/Engine/Plugins/Media/AvfMedia/AvfMedia.uplugin b/Engine/Plugins/Media/AvfMedia/AvfMedia.uplugin
--- a/Engine/Plugins/Media/AvfMedia/AvfMedia.uplugin
+++ b/Engine/Plugins/Media/AvfMedia/AvfMedia.uplugin
@@ -1,13 +1,14 @@
 {
 	"FileVersion" : 3,
-	"Version" : 1,
-	"VersionName" : "1.0",
+	"Version" : 2,
+	"VersionName" : "2.0",
 	"FriendlyName" : "AVF Media Player",
-	"Description" : "Implements a media player using the AV Foundation framework.",
+	"Description" : "Implements a media player using Apple AV Foundation.",
 	"Category" : "Media Players",
 	"CreatedBy" : "Epic Games, Inc.",
 	"CreatedByURL" : "http://epicgames.com",
 	"DocsURL" : "https://forums.unrealengine.com/showthread.php?46879-Media-Framework-Documentation-for-4-5-Preview",
+	"EngineVersion" : "4.13.0",
 	"MarketplaceURL" : "",
 	"SupportURL" : "",
 	"EnabledByDefault" : true,
@@ -20,11 +21,18 @@
 			"Name" : "AvfMedia",
 			"Type" : "Runtime",
 			"LoadingPhase" : "PreLoadingScreen",
-			"WhitelistPlatforms" :
-			[
-				"IOS",
-				"Mac"
-			]
+			"WhitelistPlatforms" : [ "IOS", "Mac" ]
+		},
+		{
+			"Name" : "AvfMediaEditor",
+			"Type" : "Editor",
+			"LoadingPhase" : "PostEngineInit"
+		},
+		{
+			"Name" : "AvfMediaFactory",
+			"Type" : "Runtime",
+			"LoadingPhase" : "PostEngineInit",
+			"WhitelistPlatforms" : [ "IOS", "Mac", "Linux", "Win32", "Win64" ]
 		}
 	]
 }
\ No newline at end of file
diff --git a/Engine/Plugins/Media/AvfMedia/Source/AvfMedia/Private/AvfMediaModule.cpp b/Engine/Plugins/Media/AvfMedia/Source/AvfMedia/Private/AvfMediaModule.cpp
--- a/Engine/Plugins/Media/AvfMedia/Source/AvfMedia/Private/AvfMediaModule.cpp
+++ b/Engine/Plugins/Media/AvfMedia/Source/AvfMedia/Private/AvfMediaModule.cpp
@@ -1,10 +1,8 @@
 // Copyright 1998-2016 Epic Games, Inc. All Rights Reserved.
 
-#include "AvfMediaPrivatePCH.h"
-#include "IMediaModule.h"
-#include "IMediaPlayerFactory.h"
-#include "ModuleInterface.h"
-#include "ModuleManager.h"
+#include "AvfMediaPCH.h"
+#include "AvfMediaPlayer.h"
+#include "IAvfMediaModule.h"
 
 
 DEFINE_LOG_CATEGORY(LogAvfMedia);
@@ -16,106 +14,28 @@ DEFINE_LOG_CATEGORY(LogAvfMedia);
  * Implements the AVFMedia module.
  */
 class FAvfMediaModule
-	: public IModuleInterface
-	, public IMediaPlayerFactory
+	: public IAvfMediaModule
 {
 public:
 
 	/** Default constructor. */
-	FAvfMediaModule()
-		: Initialized(false)
-	{ }
+	FAvfMediaModule() { }
 
 public:
 
-	// IModuleInterface interface
-
-	virtual void StartupModule() override
-	{
-		// load required libraries
-		IMediaModule* MediaModule = FModuleManager::LoadModulePtr<IMediaModule>("Media");
-
-		if (MediaModule == nullptr)
-		{
-			UE_LOG(LogAvfMedia, Warning, TEXT("Failed to load Media module"));
-
-			return;
-		}
-
-		// initialize supported media formats
-		SupportedFileTypes.Add(TEXT("avi"), LOCTEXT("FormatAvi", "Audio Video Interleave File"));
-		SupportedFileTypes.Add(TEXT("m4v"), LOCTEXT("FormatM4v", "Apple MPEG-4 Video"));
-		SupportedFileTypes.Add(TEXT("mov"), LOCTEXT("FormatMov", "Apple QuickTime Movie"));
-		SupportedFileTypes.Add(TEXT("mp4"), LOCTEXT("FormatMp4", "MPEG-4 Movie"));
-
-		// register factory
-		MediaModule->RegisterPlayerFactory(*this);
-
-		Initialized = true;
-	}
-
-	virtual void ShutdownModule() override
-	{
-		if (!Initialized)
-		{
-			return;
-		}
-
-		Initialized = false;
-
-		// unregister video player factory
-		IMediaModule* MediaModule = FModuleManager::GetModulePtr<IMediaModule>("Media");
-
-		if (MediaModule != nullptr)
-            
-		{
-			MediaModule->UnregisterPlayerFactory(*this);
-		}
-	}
-
-public:
-
-	// IMediaPlayerFactory interface
+	//~ IAvfMediaModule interface
 
 	virtual TSharedPtr<IMediaPlayer> CreatePlayer() override
 	{
-		if (Initialized)
-		{
-			return MakeShareable(new FAvfMediaPlayer());
-		}
-
-		return nullptr;
-	}
-
-	virtual const FMediaFileTypes& GetSupportedFileTypes() const override
-	{
-		return SupportedFileTypes;
+		return MakeShareable(new FAvfMediaPlayer());
 	}
 
-	virtual bool SupportsUrl(const FString& Url) const override
-	{
-		return SupportedFileTypes.Contains(FPaths::GetExtension(Url));
-	}
-
-protected:
-
-	/**
-	 * Loads all required Windows libraries.
-	 *
-	 * @return true on success, false otherwise.
-	 */
-	bool LoadRequiredLibraries()
-    {
-		return true;
-	}
-
-private:
+public:
 
-	/** Whether the module has been initialized. */
-	bool Initialized;
+	//~ IModuleInterface interface
 
-	/** The collection of supported media file types. */
-	FMediaFileTypes SupportedFileTypes;
+	virtual void StartupModule() override { }
+	virtual void ShutdownModule() override { }
 };
 
 
diff --git a/Engine/Plugins/Media/AvfMedia/Source/AvfMedia/Private/AvfMediaPCH.h b/Engine/Plugins/Media/AvfMedia/Source/AvfMedia/Private/AvfMediaPCH.h
new file mode 100644
--- /dev/null
+++ b/Engine/Plugins/Media/AvfMedia/Source/AvfMedia/Private/AvfMediaPCH.h
@@ -0,0 +1,22 @@
+// Copyright 1998-2016 Epic Games, Inc. All Rights Reserved.
+
+#pragma once
+
+
+#include "Runtime/Core/Public/Core.h"
+#include "Runtime/Core/Public/Modules/ModuleManager.h"
+#include "Runtime/Media/Public/IMediaAudioSink.h"
+#include "Runtime/Media/Public/IMediaModule.h"
+#include "Runtime/Media/Public/IMediaStringSink.h"
+#include "Runtime/Media/Public/IMediaPlayerFactory.h"
+#include "Runtime/Media/Public/IMediaTextureSink.h"
+
+#import <AVFoundation/AVFoundation.h>
+
+#if PLATFORM_MAC
+	#include "CocoaThread.h"
+#endif
+
+
+/** Log category for the AvfMedia module. */
+DECLARE_LOG_CATEGORY_EXTERN(LogAvfMedia, Log, All);
diff --git a/Engine/Plugins/Media/AvfMedia/Source/AvfMedia/Private/Player/AvfMediaPlayer.cpp b/Engine/Plugins/Media/AvfMedia/Source/AvfMedia/Private/Player/AvfMediaPlayer.cpp
--- a/Engine/Plugins/Media/AvfMedia/Source/AvfMedia/Private/Player/AvfMediaPlayer.cpp
+++ b/Engine/Plugins/Media/AvfMedia/Source/AvfMedia/Private/Player/AvfMediaPlayer.cpp
@@ -1,18 +1,17 @@
 // Copyright 1998-2016 Epic Games, Inc. All Rights Reserved.
 
-#include "AvfMediaPrivatePCH.h"
-
-#if PLATFORM_MAC
-#include "CocoaThread.h"
-#endif
+#include "AvfMediaPCH.h"
+#include "AvfMediaPlayer.h"
 
 
 /* FMediaHelper implementation
  *****************************************************************************/
+
 @implementation FMediaHelper
 @synthesize bIsPlayerItemReady;
 @synthesize MediaPlayer;
 
+
 -(FMediaHelper*) initWithMediaPlayer:(AVPlayer*)InPlayer
 {
 	MediaPlayer = InPlayer;
@@ -25,13 +24,13 @@
 
 /** Listener for changes in our media classes properties. */
 - (void) observeValueForKeyPath:(NSString*)keyPath
-					ofObject:	(id)object
-					change:		(NSDictionary*)change
-					context:	(void*)context
+		ofObject:	(id)object
+		change:		(NSDictionary*)change
+		context:	(void*)context
 {
-	if( [keyPath isEqualToString:@"status"] )
+	if ([keyPath isEqualToString:@"status"])
 	{
-		if( object == [MediaPlayer currentItem] )
+		if (object == [MediaPlayer currentItem])
 		{
 			bIsPlayerItemReady = ([MediaPlayer currentItem].status == AVPlayerItemStatusReadyToPlay);
 		}
@@ -56,16 +55,62 @@ FAvfMediaPlayer::FAvfMediaPlayer()
     CurrentTime = FTimespan::Zero();
     Duration = FTimespan::Zero();
 	MediaUrl = FString();
+	ShouldLoop = false;
     
-    PlayerItem = nil;
-    MediaPlayer = nil;
 	MediaHelper = nil;
+    MediaPlayer = nil;
+	PlayerItem = nil;
 	
 	CurrentRate = 0.0f;
 }
 
 
-/* IMediaInfo interface
+/* FTickerObjectBase interface
+ *****************************************************************************/
+
+bool FAvfMediaPlayer::Tick(float DeltaTime)
+{
+	return false;
+	/*
+	if (!IsPlaying() || !IsReady() && !ReachedEnd())
+	{
+		return true;
+	}
+
+	for (IMediaVideoTrackRef& VideoTrack : VideoTracks)
+	{
+		FAvfMediaVideoTrack& AVFTrack = (FAvfMediaVideoTrack&)VideoTrack.Get();
+
+		if (!AVFTrack.ReadFrameAtTime([[MediaPlayer currentItem] currentTime]))
+		{
+			if (ReachedEnd())
+			{
+				MediaEvent.Broadcast(EMediaEvent::PlaybackEndReached);
+
+				if (ShouldLoop)
+				{
+					Seek(FTimespan(0));
+					SetRate(CurrentRate);
+				}
+			}
+			else
+			{
+				UE_LOG(LogAvfMedia, Display, TEXT("Failed to read video track for "), *MediaUrl);
+				
+				return false;
+			}
+		}
+		else
+		{
+			CurrentTime = FTimespan::FromSeconds(CMTimeGetSeconds([[MediaPlayer currentItem] currentTime]));
+		}
+	}
+	
+    return true;*/
+}
+
+
+/* IMediaControls interface
  *****************************************************************************/
 
 FTimespan FAvfMediaPlayer::GetDuration() const
@@ -74,24 +119,73 @@ FTimespan FAvfMediaPlayer::GetDuration() const
 }
 
 
-TRange<float> FAvfMediaPlayer::GetSupportedRates( EMediaPlaybackDirections Direction, bool Unthinned ) const
+float FAvfMediaPlayer::GetRate() const
 {
-	float MaxRate = 1.0f;
-	float MinRate = 0.0f;
-    
-	return TRange<float>(MinRate, MaxRate);
+	return CurrentRate;
 }
 
 
-FString FAvfMediaPlayer::GetUrl() const
+EMediaState FAvfMediaPlayer::GetState() const
 {
-	return MediaUrl;
+	return EMediaState::Error; // @todo AvfMedia: implement state management
+}
+
+
+TRange<float> FAvfMediaPlayer::GetSupportedRates(EMediaPlaybackDirections Direction, bool Unthinned) const
+{
+	if (Direction == EMediaPlaybackDirections::Reverse)
+	{
+		return TRange<float>::Empty();
+	}
+
+	return TRange<float>(0.0f, 1.0f);
+}
+
+
+FTimespan FAvfMediaPlayer::GetTime() const
+{
+	return CurrentTime;
 }
 
 
-bool FAvfMediaPlayer::SupportsRate( float Rate, bool Unthinned ) const
+bool FAvfMediaPlayer::IsLooping() const
 {
-    return GetSupportedRates(EMediaPlaybackDirections::Forward, true).Contains( Rate );
+	return ShouldLoop;
+}
+
+
+bool FAvfMediaPlayer::Seek(const FTimespan& Time)
+{
+	return false;
+}
+
+
+bool FAvfMediaPlayer::SetLooping(bool Looping)
+{
+	ShouldLoop = Looping;
+
+	return true;
+}
+
+
+bool FAvfMediaPlayer::SetRate(float Rate)
+{
+	CurrentRate = Rate;
+	[MediaPlayer setRate : CurrentRate];
+
+	MediaEvent.Broadcast(
+		FMath::IsNearlyZero(Rate)
+		? EMediaEvent::PlaybackSuspended
+		: EMediaEvent::PlaybackResumed
+	);
+
+	return true;
+}
+
+
+bool FAvfMediaPlayer::SupportsRate(float Rate, bool Unthinned) const
+{
+	return GetSupportedRates(EMediaPlaybackDirections::Forward, true).Contains(Rate);
 }
 
 
@@ -112,308 +206,179 @@ bool FAvfMediaPlayer::SupportsSeeking() const
 
 void FAvfMediaPlayer::Close()
 {
-    CurrentTime = 0;
+	// @todo trepka: return here if already closed
+
+	CurrentTime = 0;
 	MediaUrl = FString();
-    
-    if( PlayerItem != nil )
-    {
-		[PlayerItem removeObserver:MediaHelper forKeyPath:@"status"];
+
+	if (PlayerItem != nil)
+	{
+		[PlayerItem removeObserver : MediaHelper forKeyPath : @"status"];
 		[PlayerItem release];
-        PlayerItem = nil;
-    }
+		PlayerItem = nil;
+	}
 
-	if( MediaHelper )
+	if (MediaHelper)
 	{
 		[MediaHelper release];
 		MediaHelper = nil;
 	}
-		
-	AudioTracks.Reset();
-	CaptionTracks.Reset();
-	VideoTracks.Reset();
 
-	MediaEvent.Broadcast(EMediaEvent::TracksChanged);  
-    Duration = CurrentTime = FTimespan::Zero();
-    
+	Tracks.Reset();
+	MediaEvent.Broadcast(EMediaEvent::TracksChanged);
+
+	Duration = CurrentTime = FTimespan::Zero();
+
 #if PLATFORM_MAC
-    GameThreadCall(^{
+	GameThreadCall(^{
 #elif PLATFORM_IOS
-    // Report back to the game thread whether this succeeded.
-    [FIOSAsyncTask CreateTaskWithBlock : ^ bool(void){
+		// report back to the game thread whether this succeeded
+		[FIOSAsyncTask CreateTaskWithBlock : ^ bool(void) {
 #endif
-        // Displatch on the gamethread that we have closed the video.
-        MediaEvent.Broadcast(EMediaEvent::MediaClosed);
-         
+		// dispatch on the game thread that we have closed the video
+			MediaEvent.Broadcast(EMediaEvent::MediaClosed);
 #if PLATFORM_IOS
-        return true;
-    }];
+			return true;
+		}];
 #elif PLATFORM_MAC
-    });
+	});
 #endif
-
 }
 
 
-const TArray<IMediaAudioTrackRef>& FAvfMediaPlayer::GetAudioTracks() const
+IMediaControls& FAvfMediaPlayer::GetControls()
 {
-	return AudioTracks;
+	return *this;
 }
 
 
-const TArray<IMediaCaptionTrackRef>& FAvfMediaPlayer::GetCaptionTracks() const
+FString FAvfMediaPlayer::GetInfo() const
 {
-	return CaptionTracks;
+	return TEXT("AvfMedia media information not implemented yet");
 }
 
 
-const IMediaInfo& FAvfMediaPlayer::GetMediaInfo() const
+IMediaOutput& FAvfMediaPlayer::GetOutput()
 {
-	return *this;
+	return Tracks;
 }
 
 
-float FAvfMediaPlayer::GetRate() const
+FString FAvfMediaPlayer::GetStats() const
 {
-    return CurrentRate;
+	return TEXT("AvfMedia stats information not implemented yet");
 }
 
 
-FTimespan FAvfMediaPlayer::GetTime() const
+IMediaTracks& FAvfMediaPlayer::GetTracks()
 {
-    return CurrentTime;
+	return Tracks;
 }
 
 
-const TArray<IMediaVideoTrackRef>& FAvfMediaPlayer::GetVideoTracks() const
+FString FAvfMediaPlayer::GetUrl() const
 {
-	return VideoTracks;
+	return MediaUrl;
 }
 
 
-bool FAvfMediaPlayer::IsLooping() const 
+bool FAvfMediaPlayer::Open(const FString& Url, const IMediaOptions& Options)
 {
-    return bLoop;
-}
+	Close();
 
+	// open media file
+	FString Ext = FPaths::GetExtension(Url);
+	FString BFNStr = FPaths::GetBaseFilename(Url);
 
-bool FAvfMediaPlayer::IsPaused() const
-{
-    return FMath::IsNearlyZero( CurrentRate );
-}
-
+#if PLATFORM_MAC
+	NSURL* nsMediaUrl = [NSURL fileURLWithPath : Url.GetNSString()];
+#elif PLATFORM_IOS
+	NSURL* nsMediaUrl = [[NSBundle mainBundle] URLForResource:BFNStr.GetNSString() withExtension : Ext.GetNSString()];
+#endif
 
-bool FAvfMediaPlayer::IsPlaying() const
-{
-	if ((AudioTracks.Num() > 0) && (CaptionTracks.Num() > 0) && (VideoTracks.Num() > 0))
+	if (nsMediaUrl == nil)
 	{
+		UE_LOG(LogAvfMedia, Error, TEXT("Failed to open Media file:"), *Url);
+
 		return false;
 	}
 
-    return (MediaPlayer != nil) && !FMath::IsNearlyZero(CurrentRate) && !ReachedEnd();
-}
-
+	// create player instance
+	MediaUrl = FPaths::GetCleanFilename(Url);
+	MediaPlayer = [[AVPlayer alloc] init];
 
-bool FAvfMediaPlayer::IsReady() const
-{
-    // To be ready, we need the AVPlayer setup
-	if ((MediaPlayer == nil) || ([MediaPlayer status] != AVPlayerStatusReadyToPlay))
+	if (!MediaPlayer)
 	{
+		UE_LOG(LogAvfMedia, Error, TEXT("Failed to create instance of an AVPlayer"));
+
 		return false;
 	}
 
-	// A player item setup and ready to stream,
-	if ((MediaHelper == nil) || (![MediaHelper bIsPlayerItemReady]))
+	// create player item
+	MediaHelper = [[FMediaHelper alloc] initWithMediaPlayer:MediaPlayer];
+	check(MediaHelper != nil);
+
+	PlayerItem = [AVPlayerItem playerItemWithURL : nsMediaUrl];
+
+	if (PlayerItem == nil)
 	{
+		UE_LOG(LogAvfMedia, Error, TEXT("Failed to open player item with Url:"), *Url);
+
 		return false;
 	}
-    
-    // and all tracks to be setup and ready
-    for( const IMediaAudioTrackRef& AudioTrack : AudioTracks )
-    {
-        FAvfMediaTrack& AVFTrack = (FAvfMediaTrack&)AudioTrack.Get();
-		if (!AVFTrack.IsReady())
-		{
-			return false;
-        }
-    }
-
-	for( const IMediaCaptionTrackRef& CaptionTrack : CaptionTracks )
-    {
-        FAvfMediaTrack& AVFTrack = (FAvfMediaTrack&)CaptionTrack.Get();
-		if (!AVFTrack.IsReady())
-		{
-			return false;
-        }
-    }
-
-	for( const IMediaVideoTrackRef& VideoTrack : VideoTracks )
-    {
-        FAvfMediaVideoTrack& AVFTrack = (FAvfMediaVideoTrack&)VideoTrack.Get();
-		if (!AVFTrack.IsReady())
-		{
-			return false;
-        }
-    }
 
-	return true;
-}
+	// load tracks
+	[PlayerItem retain];
+	[PlayerItem addObserver : MediaHelper forKeyPath : @"status" options : 0 context : nil];
 
+	Duration = FTimespan::FromSeconds(CMTimeGetSeconds(PlayerItem.asset.duration));
 
-/* FAvfMediaPlayer implementation
- *****************************************************************************/
+	[[PlayerItem asset] loadValuesAsynchronouslyForKeys:@[@"tracks"] completionHandler:^
+	{
+		NSError* Error = nil;
 
-bool FAvfMediaPlayer::Open( const FString& Url )
-{
-    Close();
-    
-    bool bSuccessfullyOpenedMovie = false;
+		if ([[PlayerItem asset] statusOfValueForKey:@"tracks" error : &Error] == AVKeyValueStatusLoaded)
+		{
+//			Tracks.Initialize(PlayerItem); // @todo trepka: fix me
 
-    FString Ext = FPaths::GetExtension(Url);
-    FString BFNStr = FPaths::GetBaseFilename(Url);
-    
 #if PLATFORM_MAC
-    NSURL* nsMediaUrl = [NSURL fileURLWithPath:Url.GetNSString()];
+			GameThreadCall(^{
 #elif PLATFORM_IOS
-    NSURL* nsMediaUrl = [[NSBundle mainBundle] URLForResource:BFNStr.GetNSString() withExtension:Ext.GetNSString()];
+				// report back to the game thread whether this succeeded.
+				[FIOSAsyncTask CreateTaskWithBlock : ^ bool(void) {
 #endif
-    
-    if( nsMediaUrl != nil )
-    {
-        MediaUrl = FPaths::GetCleanFilename(Url);
-        
-        MediaPlayer = [[AVPlayer alloc] init];
-        if( MediaPlayer )
-        {
-			MediaHelper = [[FMediaHelper alloc] initWithMediaPlayer: MediaPlayer];
-			check( MediaHelper != nil );
-			
-            PlayerItem = [AVPlayerItem playerItemWithURL: nsMediaUrl];
-            if( PlayerItem != nil )
-            {
-				[PlayerItem retain];
-				[PlayerItem addObserver:MediaHelper forKeyPath:@"status" options:0 context:nil];
-                Duration = FTimespan::FromSeconds( CMTimeGetSeconds( PlayerItem.asset.duration ) );
-
-                [[PlayerItem asset] loadValuesAsynchronouslyForKeys: @[@"tracks"] completionHandler:^
-                {
-                    NSError* Error = nil;
-                    if( [[PlayerItem asset] statusOfValueForKey:@"tracks" error:&Error] == AVKeyValueStatusLoaded )
-                    {
-                        NSArray* AssetVideoTracks = [[PlayerItem asset] tracksWithMediaType:AVMediaTypeVideo];
-                        if( AssetVideoTracks != nullptr && AssetVideoTracks.count > 0 )
-                        {
-                            AVAssetTrack* VideoTrack = [AssetVideoTracks objectAtIndex: 0];
-							if (VideoTrack != nullptr)
-							{
-								FAvfMediaVideoTrack* NewTrack = new FAvfMediaVideoTrack(VideoTrack);
-								if (NewTrack != nullptr && NewTrack->IsReady())
-								{
-									VideoTracks.Add( MakeShareable( NewTrack ) );
-#if PLATFORM_MAC
-									GameThreadCall(^{
-#elif PLATFORM_IOS
-									// Report back to the game thread whether this succeeded.
-									[FIOSAsyncTask CreateTaskWithBlock : ^ bool(void){
-#endif
-										 // Displatch on the gamethread that we have opened the video.
-										MediaEvent.Broadcast(EMediaEvent::TracksChanged);
-										MediaEvent.Broadcast(EMediaEvent::MediaOpened);
+					// dispatch on the gamethread that we have opened the video.
+					MediaEvent.Broadcast(EMediaEvent::TracksChanged);
+					MediaEvent.Broadcast(EMediaEvent::MediaOpened);
 #if PLATFORM_IOS
-										return true;
-									}];
+					return true;
+				}];
 #elif PLATFORM_MAC
-									});
+			});
 #endif
-								}
-								else if(!NewTrack->IsReady())
-								{
-									delete NewTrack;
-								}
-							}
-                        }
-                    }
-                    else if(Error != nullptr)
-                    {
-                        NSDictionary *userInfo = [Error userInfo];
-                        NSString *errstr = [[userInfo objectForKey : NSUnderlyingErrorKey] localizedDescription];
-                        UE_LOG(LogAvfMedia, Warning, TEXT("Failed to load video tracks. [%s]"), *FString(errstr));
-                    }
-                }];
-            
-                [MediaPlayer replaceCurrentItemWithPlayerItem: PlayerItem];
-            
-                [[MediaPlayer currentItem] seekToTime:kCMTimeZero];
-                MediaPlayer.rate = 0.0;
-				CurrentTime = FTimespan::Zero();
-            
-                bSuccessfullyOpenedMovie = true;
-            }
-            else
-            {
-                UE_LOG(LogAvfMedia, Warning, TEXT("Failed to open player item with Url:"), *Url);
-            }
-        }
-        else
-        {
-            UE_LOG(LogAvfMedia, Error, TEXT("Failed to create instance of an AVPlayer"));
-        }
-    }
-    else
-    {
-        UE_LOG(LogAvfMedia, Warning, TEXT("Failed to open Media file:"), *Url);
-    }
-    
-    
-    return bSuccessfullyOpenedMovie;
-}
-
-
-bool FAvfMediaPlayer::Open( const TSharedRef<FArchive, ESPMode::ThreadSafe>& Archive, const FString& OriginalUrl )
-{
-    return false;
-}
+		}
+		else if (Error != nullptr)
+		{
+			NSDictionary *userInfo = [Error userInfo];
+			NSString *errstr = [[userInfo objectForKey : NSUnderlyingErrorKey] localizedDescription];
 
+			UE_LOG(LogAvfMedia, Warning, TEXT("Failed to load video tracks. [%s]"), *FString(errstr));
+		}
+	}];
 
-bool FAvfMediaPlayer::Seek( const FTimespan& Time )
-{
-    // We need to find a suitable way to seek using avplayer and AVAssetReader,
-    CurrentTime = Time;
-	
-	double TotalSeconds = CurrentTime.GetTotalSeconds();
-	CMTime CurrentTimeInSeconds = CMTimeMake(TotalSeconds, 1);
-	
-	[MediaPlayer seekToTime:CurrentTimeInSeconds];
-	[[MediaPlayer currentItem] seekToTime:CurrentTimeInSeconds];
-	
-    for (IMediaVideoTrackRef& VideoTrack : VideoTracks)
-    {
-        FAvfMediaVideoTrack& AVFTrack = (FAvfMediaVideoTrack&)VideoTrack.Get();
-        AVFTrack.SeekToTime(CurrentTimeInSeconds);
-    }
+	[MediaPlayer replaceCurrentItemWithPlayerItem : PlayerItem];
+	[[MediaPlayer currentItem] seekToTime:kCMTimeZero];
 
-    return false;
+	MediaPlayer.rate = 0.0;
+	CurrentTime = FTimespan::Zero();
+		
+	return true;
 }
 
 
-bool FAvfMediaPlayer::SetLooping( bool bInLooping )
+bool FAvfMediaPlayer::Open(const TSharedRef<FArchive, ESPMode::ThreadSafe>& Archive, const FString& OriginalUrl, const IMediaOptions& Options)
 {
-    bLoop = bInLooping;
-	return bLoop;
-}
-
-
-bool FAvfMediaPlayer::SetRate( float Rate )
-{
-    CurrentRate = Rate;
-    [MediaPlayer setRate: CurrentRate];
-
-	MediaEvent.Broadcast(
-		FMath::IsNearlyZero(Rate)
-			? EMediaEvent::PlaybackSuspended
-			: EMediaEvent::PlaybackResumed
-	);
-
-    return true;
+	return false; // not supported
 }
 
 
@@ -421,69 +386,16 @@ bool FAvfMediaPlayer::SetRate( float Rate )
  *****************************************************************************/
 
 bool FAvfMediaPlayer::ReachedEnd() const
-{
-	bool bHasReachedEnd = false;
-	for( const IMediaVideoTrackRef& VideoTrack : VideoTracks )
+{/*
+	for (const IMediaVideoTrackRef& VideoTrack : VideoTracks)
 	{
 		FAvfMediaVideoTrack& AVFTrack = (FAvfMediaVideoTrack&)VideoTrack.Get();
+
 		if (AVFTrack.ReachedEnd())
 		{
-			bHasReachedEnd = true;
+			return true;
 		}
-	}
-	return bHasReachedEnd;
-}
-
-
-bool FAvfMediaPlayer::ShouldTick() const
-{
-	return IsReady() || ReachedEnd();
-}
-
-
-/* FTickerObjectBase interface
- *****************************************************************************/
-
-bool FAvfMediaPlayer::Tick( float DeltaTime )
-{
-	if (!ShouldTick() || !IsPlaying())
-	{
-		return true;
-	}
-
-	for (IMediaVideoTrackRef& VideoTrack : VideoTracks)
-	{
-		FAvfMediaVideoTrack& AVFTrack = (FAvfMediaVideoTrack&)VideoTrack.Get();
-
-		if (!AVFTrack.ReadFrameAtTime([[MediaPlayer currentItem] currentTime]))
-		{
-			if (ReachedEnd())
-			{
-				MediaEvent.Broadcast(EMediaEvent::PlaybackEndReached);
+	}*/
 
-				if (bLoop)
-				{
-					Seek(FTimespan(0));
-					SetRate(CurrentRate);
-				}
-			}
-			else
-			{
-				UE_LOG(LogAvfMedia, Display, TEXT("Failed to read video track for "), *MediaUrl);
-				
-				return false;
-			}
-		}
-		else
-		{
-			CurrentTime = FTimespan::FromSeconds( CMTimeGetSeconds([[MediaPlayer currentItem] currentTime]) );
-		}
-	}
-	
-    return true;
+	return false;
 }
-
-
-/* FAvfMediaPlayer implementation
- *****************************************************************************/
-
diff --git a/Engine/Plugins/Media/AvfMedia/Source/AvfMedia/Private/Player/AvfMediaPlayer.h b/Engine/Plugins/Media/AvfMedia/Source/AvfMedia/Private/Player/AvfMediaPlayer.h
--- a/Engine/Plugins/Media/AvfMedia/Source/AvfMedia/Private/Player/AvfMediaPlayer.h
+++ b/Engine/Plugins/Media/AvfMedia/Source/AvfMedia/Private/Player/AvfMediaPlayer.h
@@ -2,16 +2,22 @@
 
 #pragma once
 
-#include "IMediaInfo.h"
+#include "AvfMediaTracks.h"
+#include "IMediaControls.h"
 #include "IMediaPlayer.h"
 
 
+@class AVPlayer;
+@class AVPlayerItem;
+
+
 /**
  * Cocoa class that can help us with reading player item information.
  */
 @interface FMediaHelper : NSObject
 {
 };
+
 /** We should only initiate a helper with a media player */
 -(FMediaHelper*) initWithMediaPlayer:(AVPlayer*)InPlayer;
 /** Destructor */
@@ -30,9 +36,9 @@
  * Implements a media player using the AV framework.
  */
 class FAvfMediaPlayer
-	: public IMediaInfo
+	: public FTickerObjectBase
+	, public IMediaControls
 	, public IMediaPlayer
-    , public FTickerObjectBase
 {
 public:
 
@@ -44,97 +50,92 @@ public:
 
 public:
 
-	// IMediaInfo interface
+	//~ FTickerObjectBase interface
+
+	virtual bool Tick(float DeltaTime) override;
+
+public:
+
+	//~ IMediaControls interface
 
 	virtual FTimespan GetDuration() const override;
-	virtual TRange<float> GetSupportedRates( EMediaPlaybackDirections Direction, bool Unthinned ) const override;
-	virtual FString GetUrl() const override;
-	virtual bool SupportsRate( float Rate, bool Unthinned ) const override;
+	virtual float GetRate() const override;
+	virtual EMediaState GetState() const override;
+	virtual TRange<float> GetSupportedRates(EMediaPlaybackDirections Direction, bool Unthinned) const override;
+	virtual FTimespan GetTime() const override;
+	virtual bool IsLooping() const override;
+	virtual bool Seek(const FTimespan& Time) override;
+	virtual bool SetLooping(bool Looping) override;
+	virtual bool SetRate(float Rate) override;
+	virtual bool SupportsRate(float Rate, bool Unthinned) const override;
 	virtual bool SupportsScrubbing() const override;
 	virtual bool SupportsSeeking() const override;
 
 public:
 
-	// IMediaPlayer interface
+	//~ IMediaPlayer interface
 
 	virtual void Close() override;
-	virtual const TArray<IMediaAudioTrackRef>& GetAudioTracks() const override;
-	virtual const TArray<IMediaCaptionTrackRef>& GetCaptionTracks() const override;
-	virtual const IMediaInfo& GetMediaInfo() const override;
-	virtual float GetRate() const override;
-	virtual FTimespan GetTime() const override;
-	virtual const TArray<IMediaVideoTrackRef>& GetVideoTracks() const override;
-	virtual bool IsLooping() const override;
-	virtual bool IsPaused() const override;
-	virtual bool IsPlaying() const override;
-	virtual bool IsReady() const override;
-	virtual bool Open( const FString& Url ) override;
-	virtual bool Open( const TSharedRef<FArchive, ESPMode::ThreadSafe>& Archive, const FString& OriginalUrl ) override;
-	virtual bool Seek( const FTimespan& Time ) override;
-	virtual bool SetLooping( bool Looping ) override;
-	virtual bool SetRate( float Rate ) override;
-
-	DECLARE_DERIVED_EVENT(FAvfMediaPlayer, IMediaPlayer::FOnMediaEvent, FOnMediaEvent);
+	virtual IMediaControls& GetControls() override;
+	virtual FString GetInfo() const override;
+	virtual IMediaOutput& GetOutput() override;
+	virtual FString GetStats() const override;
+	virtual IMediaTracks& GetTracks() override;
+	virtual FString GetUrl() const override;
+	virtual bool Open(const FString& Url, const IMediaOptions& Options) override;
+	virtual bool Open(const TSharedRef<FArchive, ESPMode::ThreadSafe>& Archive, const FString& OriginalUrl, const IMediaOptions& Options) override;
+
+	DECLARE_DERIVED_EVENT(FVlcMediaPlayer, IMediaPlayer::FOnMediaEvent, FOnMediaEvent);
 	virtual FOnMediaEvent& OnMediaEvent() override
 	{
 		return MediaEvent;
 	}
-
-public:
-
-    // FTickerObjectBase interface
-
-    virtual bool Tick(float DeltaTime) override;
 	
 protected:
 
 	/**
-	 * Has the video completed it's playthrough
+	 * Whether the video finished playback.
+	 *
+	 * @return true if playback has finished, false otherwise.
 	 */
 	bool ReachedEnd() const;
 
 	/**
 	 * Whether the media player should tick in this frame.
 	 *
-	 * @return true if the player should advance, false otherwise
+	 * @return true if the player should advance, false otherwise.
 	 */
-	bool ShouldTick() const;
+//	bool ShouldTick() const;
 
 private:
 
-	/** The available audio tracks. */
-	TArray<IMediaAudioTrackRef> AudioTracks;
-
-	/** The available caption tracks. */
-	TArray<IMediaCaptionTrackRef> CaptionTracks;
+	/** The current playback rate. */
+	float CurrentRate;
 
-    /** The AVFoundation media player */
-    AVPlayer* MediaPlayer;
+	/** The current time of the playback. */
+	FTimespan CurrentTime;
 
-    /** The player item which the media player uses to progress. */
-    AVPlayerItem* PlayerItem;
-
-    /** The duration of the media. */
+	/** The duration of the media. */
     FTimespan Duration;
 
-    /** The current time of the playback. */
-    FTimespan CurrentTime;
-
-    /** The URL of the currently opened media. */
-    FString MediaUrl;
-
-    /** The current playback rate. */
-    float CurrentRate;
-
 	/** Holds an event delegate that is invoked when a media event occurred. */
 	FOnMediaEvent MediaEvent;
 
 	/** Cocoa helper object we can use to keep track of ns property changes in our media items */
 	FMediaHelper* MediaHelper;
-
-	/** The available video tracks. */
-	TArray<IMediaVideoTrackRef> VideoTracks;
     
-    /** Should the video loop to the beginning at completion */
-    bool bLoop;
+	/** The AVFoundation media player */
+	AVPlayer* MediaPlayer;
+
+	/** The URL of the currently opened media. */
+	FString MediaUrl;
+
+	/** The player item which the media player uses to progress. */
+	AVPlayerItem* PlayerItem;
+
+	/** Should the video loop to the beginning at completion */
+    bool ShouldLoop;
+
+	/** The media track collection. */
+	FAvfMediaTracks Tracks;
 };
diff --git a/Engine/Plugins/Media/AvfMedia/Source/AvfMedia/Private/Player/AvfMediaTracks.cpp b/Engine/Plugins/Media/AvfMedia/Source/AvfMedia/Private/Player/AvfMediaTracks.cpp
new file mode 100644
--- /dev/null
+++ b/Engine/Plugins/Media/AvfMedia/Source/AvfMedia/Private/Player/AvfMediaTracks.cpp
@@ -0,0 +1,560 @@
+// Copyright 1998-2016 Epic Games, Inc. All Rights Reserved.
+
+#include "AvfMediaPCH.h"
+#include "AvfMediaTracks.h"
+
+
+/* FAvfMediaTracks structors
+ *****************************************************************************/
+
+FAvfMediaTracks::FAvfMediaTracks()
+	: AudioSink(nullptr)
+	, CaptionSink(nullptr)
+	, VideoSink(nullptr)
+	, PlayerItem(nullptr)
+	, SelectedAudioTrack(INDEX_NONE)
+	, SelectedCaptionTrack(INDEX_NONE)
+	, SelectedVideoTrack(INDEX_NONE)
+{ }
+
+
+FAvfMediaTracks::~FAvfMediaTracks()
+{
+	Reset();
+}
+
+
+/* FAvfMediaTracks interface
+ *****************************************************************************/
+
+void FAvfMediaTracks::Flush()
+{
+	FScopeLock Lock(&CriticalSection);
+
+	if (AudioSink != nullptr)
+	{
+		AudioSink->FlushAudioSink();
+	}
+}
+
+
+void FAvfMediaTracks::Initialize(AVPlayerItem& InPlayerItem)
+{
+	Reset();
+
+	FScopeLock Lock(&CriticalSection);
+
+	PlayerItem = &InPlayerItem;
+/*
+	// initialize tracks
+	[[PlayerItem asset] loadValuesAsynchronouslyForKeys:@[@"tracks"] completionHandler:^
+	{
+		NSError* Error = nil;
+
+		if ([[PlayerItem asset] statusOfValueForKey:@"tracks" error : &Error] == AVKeyValueStatusLoaded)
+		{
+			NSArray* AssetTracks = [[PlayerItem asset] tracks];
+
+			for (AVAssetTrack* AssetTrack in AssetTracks)
+			{
+				// create asset reader
+				AVAssetReader Reader = [[AVAssetReader alloc] initWithAsset: [AssetTrack asset] error:&Error];
+
+				if (Error != nil)
+				{
+					FString ErrorStr([nsError localizedDescription]);
+					UE_LOG(LogAvfMedia, Error, TEXT("Failed to create asset reader for track %i: %s"), Track.trackID, *ErrorStr);
+
+					continue;
+				}
+
+				// create track
+				FTrack* Track = nullptr;
+				String mediaType = AssetTrack->mediaType;
+
+				if (mediaType == AVMediaTypeAudio)
+				{
+					// @todo trepka: create & initialize audio output
+
+					int32 TrackIndex = AudioTracks.AddDefaulted()];
+					Track = &AudioTracks[TrackIndex];
+
+					Track->Name = FString::Printf(TEXT("Audio Track %i"), TrackIndex);
+//					Track->Output = Output;
+				}
+				else if ((mediaType == AVMediaTypeClosedCaption) || (mediaType == AVMediaTypeClosedSubtitle))
+				{
+					// @todo trepka: create & initialize caption output
+
+					int32 TrackIndex = CaptionTracks.AddDefaulted()];
+					Track = &CaptionTracks[TrackIndex];
+
+					Track->Name = FString::Printf(TEXT("Caption Track %i"), TrackIndex);
+//					Track->Output = Output;
+				}
+				else if (mediaType == AVMediaTypeText)
+				{
+					// not supported by Media Framework yet
+				}
+				else if (mediaType == AVMediaTypeTimecode)
+				{
+					// not implemented yet
+				}
+				else if (mediaType == AVMediaTypeVideo)
+				{
+					NSMutableDictionary* OutputSettings = [NSMutableDictionary dictionary];
+					[OutputSettings setObject : [NSNumber numberWithInt : kCVPixelFormatType_32BGRA] forKey : (NSString*)kCVPixelBufferPixelFormatTypeKey];
+
+					AVAssetReaderTrackOutput* Output = [[AVAssetReaderTrackOutput alloc] initWithTrack:Track outputSettings : OutputSettings];
+					Output.alwaysCopiesSampleData = NO;
+					[AVReader addOutput:AVVideoOutput];
+
+					int32 TrackIndex = VideoTracks.AddDefaulted()];
+					Track = &VideoTracks[TrackIndex];
+
+					Track->Name = FString::Printf(TEXT("Video Track %i"), TrackIndex);
+					Track->Output = Output;
+				}
+
+				if (Track == nullptr)
+				{
+					continue;
+				}
+
+				Track->AssetTrack = AssetTrack;
+				Track->DisplayName = FText::FromString(Track->Name);
+				Track->Loaded = [Reader startReading];
+				Track->Reader = Reader;
+				Track->SyncStatus = Default;
+			}
+		}
+		else if (Error != nullptr)
+		{
+			NSDictionary *userInfo = [Error userInfo];
+			NSString *ErrorString = [[userInfo objectForKey : NSUnderlyingErrorKey] localizedDescription];
+
+			UE_LOG(LogAvfMedia, Warning, TEXT("Failed to load video tracks: %s"), *FString(ErrorString));
+		}
+	}];*/
+}
+
+
+void FAvfMediaTracks::Reset()
+{
+	FScopeLock Lock(&CriticalSection);
+
+	SelectedAudioTrack = INDEX_NONE;
+	SelectedCaptionTrack = INDEX_NONE;
+	SelectedVideoTrack = INDEX_NONE;
+
+	AudioTracks.Empty();
+	CaptionTracks.Empty();
+	VideoTracks.Empty();
+
+	if (AudioSink != nullptr)
+	{
+		AudioSink->ShutdownAudioSink();
+	}
+
+	if (CaptionSink != nullptr)
+	{
+		CaptionSink->ShutdownStringSink();
+	}
+
+	if (VideoSink != nullptr)
+	{
+		VideoSink->ShutdownTextureSink();
+	}
+
+	AudioTracks.Empty();
+	CaptionTracks.Empty();
+	VideoTracks.Empty();
+}
+
+
+void FAvfMediaTracks::SetPaused(bool Paused)
+{
+	FScopeLock Lock(&CriticalSection);
+
+	if (AudioSink != nullptr)
+	{
+		if (Paused)
+		{
+			AudioSink->PauseAudioSink();
+		}
+		else
+		{
+			AudioSink->ResumeAudioSink();
+		}
+	}
+}
+
+
+/* IMediaOutput interface
+ *****************************************************************************/
+
+void FAvfMediaTracks::SetAudioSink(IMediaAudioSink* Sink)
+{
+	if (Sink != AudioSink)
+	{
+		FScopeLock Lock(&CriticalSection);
+
+		if (AudioSink != nullptr)
+		{
+			AudioSink->ShutdownAudioSink();
+		}
+
+		AudioSink = Sink;	
+		InitializeAudioSink();
+	}
+}
+
+
+void FAvfMediaTracks::SetCaptionSink(IMediaStringSink* Sink)
+{
+	if (Sink != CaptionSink)
+	{
+		FScopeLock Lock(&CriticalSection);
+
+		if (CaptionSink != nullptr)
+		{
+			CaptionSink->ShutdownStringSink();
+		}
+
+		CaptionSink = Sink;
+		InitializeCaptionSink();
+	}
+}
+
+
+void FAvfMediaTracks::SetImageSink(IMediaTextureSink* Sink)
+{
+	// not supported
+}
+
+
+void FAvfMediaTracks::SetVideoSink(IMediaTextureSink* Sink)
+{
+	if (Sink != VideoSink)
+	{
+		FScopeLock Lock(&CriticalSection);
+
+		if (VideoSink != nullptr)
+		{
+			VideoSink->ShutdownTextureSink();
+		}
+
+		VideoSink = Sink;
+		InitializeVideoSink();
+	}
+}
+
+
+/* IMediaTracks interface
+ *****************************************************************************/
+
+uint32 FAvfMediaTracks::GetAudioTrackChannels(int32 TrackIndex) const
+{
+	if (!AudioTracks.IsValidIndex(TrackIndex))
+	{
+		return 0;
+	}
+
+	// @todo trepka: get channelsPerFrame from AudioTracks[TrackIndex].AssetTrack.formatDescription (CMAudioFormatDescription)
+	return 0;
+}
+
+
+uint32 FAvfMediaTracks::GetAudioTrackSampleRate(int32 TrackIndex) const
+{
+	if (!AudioTracks.IsValidIndex(TrackIndex))
+	{
+		return 0;
+	}
+
+	// @todo trepka: get sampleRate from AudioTracks[TrackIndex]AssetTrack.formatDescription (CMAudioFormatDescription)
+	return 0;
+}
+
+
+int32 FAvfMediaTracks::GetNumTracks(EMediaTrackType TrackType) const
+{
+	switch (TrackType)
+	{
+	case EMediaTrackType::Audio:
+		return AudioTracks.Num();
+	case EMediaTrackType::Caption:
+		return CaptionTracks.Num();
+	case EMediaTrackType::Video:
+		return VideoTracks.Num();
+	default:
+		return 0;
+	}
+}
+
+
+int32 FAvfMediaTracks::GetSelectedTrack(EMediaTrackType TrackType) const
+{
+	switch (TrackType)
+	{
+	case EMediaTrackType::Audio:
+		return SelectedAudioTrack;
+	case EMediaTrackType::Caption:
+		return SelectedCaptionTrack;
+	case EMediaTrackType::Video:
+		return SelectedVideoTrack;
+	default:
+		return INDEX_NONE;
+	}
+}
+
+
+FText FAvfMediaTracks::GetTrackDisplayName(EMediaTrackType TrackType, int32 TrackIndex) const
+{
+	switch (TrackType)
+	{
+	case EMediaTrackType::Audio:
+		if (AudioTracks.IsValidIndex(TrackIndex))
+		{
+			return AudioTracks[TrackIndex].DisplayName;
+		}
+		break;
+
+	case EMediaTrackType::Caption:
+		if (CaptionTracks.IsValidIndex(TrackIndex))
+		{
+			return CaptionTracks[TrackIndex].DisplayName;
+		}
+		break;
+
+	case EMediaTrackType::Video:
+		if (VideoTracks.IsValidIndex(TrackIndex))
+		{
+			return VideoTracks[TrackIndex].DisplayName;
+		}
+		break;
+
+	default:
+		break;
+	}
+
+	return FText::GetEmpty();
+}
+
+
+FString FAvfMediaTracks::GetTrackLanguage(EMediaTrackType TrackType, int32 TrackIndex) const
+{
+	if (TrackType == EMediaTrackType::Audio)
+	{
+		if (AudioTracks.IsValidIndex(TrackIndex))
+		{
+			return FString(AudioTracks[TrackIndex].AssetTrack.languageCode);
+		}
+	}
+	else if (TrackType == EMediaTrackType::Caption)
+	{
+		if (CaptionTracks.IsValidIndex(TrackIndex))
+		{
+			return FString(CaptionTracks[TrackIndex].AssetTrack.languageCode);
+		}
+	}
+	else if (TrackType == EMediaTrackType::Video)
+	{
+		if (VideoTracks.IsValidIndex(TrackIndex))
+		{
+			return FString(VideoTracks[TrackIndex].AssetTrack.languageCode);
+		}
+	}
+
+	return FString();
+}
+
+
+FString FAvfMediaTracks::GetTrackName(EMediaTrackType TrackType, int32 TrackIndex) const
+{
+	switch (TrackType)
+	{
+	case EMediaTrackType::Audio:
+		if (AudioTracks.IsValidIndex(TrackIndex))
+		{
+			return AudioTracks[TrackIndex].Name;
+		}
+		break;
+
+	case EMediaTrackType::Caption:
+		if (CaptionTracks.IsValidIndex(TrackIndex))
+		{
+			return CaptionTracks[TrackIndex].Name;
+		}
+		break;
+
+	case EMediaTrackType::Video:
+		if (VideoTracks.IsValidIndex(TrackIndex))
+		{
+			return VideoTracks[TrackIndex].Name;
+		}
+		break;
+
+	default:
+		break;
+	}
+
+	return FString();
+}
+
+
+uint32 FAvfMediaTracks::GetVideoTrackBitRate(int32 TrackIndex) const
+{
+	return VideoTracks.IsValidIndex(TrackIndex) ? [VideoTracks[TrackIndex].AssetTrack estimatedDataRate] : 0;
+}
+
+
+FIntPoint FAvfMediaTracks::GetVideoTrackDimensions(int32 TrackIndex) const
+{
+	return FIntPoint::ZeroValue; // @todo trepka: fix me
+/*	return VideoTracks.IsValidIndex(TrackIndex)
+		? FIntPoint([[VideoTracks[TrackIndex].AssetTrack naturalSize] width],
+					[[VideoTracks[TrackIndex].AssetTrack naturalSize] height])
+		: FIntPoint::ZeroValue;*/
+}
+
+
+float FAvfMediaTracks::GetVideoTrackFrameRate(int32 TrackIndex) const
+{
+	return VideoTracks.IsValidIndex(TrackIndex) ? 1.0f / [VideoTracks[TrackIndex].AssetTrack nominalFrameRate] : 0.0f;
+}
+
+
+bool FAvfMediaTracks::SelectTrack(EMediaTrackType TrackType, int32 TrackIndex)
+{
+	FScopeLock Lock(&CriticalSection);
+
+	switch (TrackType)
+	{
+	case EMediaTrackType::Audio:
+		if ((TrackIndex == INDEX_NONE) || AudioTracks.IsValidIndex(TrackIndex))
+		{
+			if (TrackIndex != SelectedAudioTrack)
+			{
+				if (SelectedAudioTrack != INDEX_NONE)
+				{
+//					PresentationDescriptor->DeselectStream(AudioTracks[SelectedAudioTrack].StreamIndex);
+					SelectedAudioTrack = INDEX_NONE;
+				}
+
+//				if ((TrackIndex != INDEX_NONE) && SUCCEEDED(PresentationDescriptor->SelectStream(AudioTracks[TrackIndex].StreamIndex)))
+				{
+					SelectedAudioTrack = TrackIndex;
+				}
+
+				if (SelectedAudioTrack == TrackIndex)
+				{
+					InitializeAudioSink();
+				}
+			}
+
+			return true;
+		}
+		break;
+
+	case EMediaTrackType::Caption:
+		if ((TrackIndex == INDEX_NONE) || CaptionTracks.IsValidIndex(TrackIndex))
+		{
+			if (TrackIndex != SelectedCaptionTrack)
+			{
+				if (SelectedCaptionTrack != INDEX_NONE)
+				{
+//					PresentationDescriptor->DeselectStream(CaptionTracks[SelectedCaptionTrack].StreamIndex);
+					SelectedCaptionTrack = INDEX_NONE;
+				}
+
+//				if ((TrackIndex != INDEX_NONE) && SUCCEEDED(PresentationDescriptor->SelectStream(CaptionTracks[TrackIndex].StreamIndex)))
+				{
+					SelectedCaptionTrack = TrackIndex;
+				}
+
+				if (SelectedCaptionTrack == TrackIndex)
+				{
+					InitializeCaptionSink();
+				}
+			}
+
+			return true;
+		}
+		break;
+
+	case EMediaTrackType::Video:
+		if ((TrackIndex == INDEX_NONE) || VideoTracks.IsValidIndex(TrackIndex))
+		{
+			if (TrackIndex != SelectedVideoTrack)
+			{
+				if (SelectedVideoTrack != INDEX_NONE)
+				{
+//					PresentationDescriptor->DeselectStream(VideoTracks[SelectedVideoTrack].StreamIndex);
+					SelectedVideoTrack = INDEX_NONE;
+				}
+
+//				if ((TrackIndex != INDEX_NONE) && SUCCEEDED(PresentationDescriptor->SelectStream(VideoTracks[TrackIndex].StreamIndex)))
+				{
+					SelectedVideoTrack = TrackIndex;
+				}
+
+				if (SelectedVideoTrack == TrackIndex)
+				{
+					InitializeVideoSink();
+				}
+			}
+
+			return true;
+		}
+		break;
+
+	default:
+		break;
+	}
+
+	return false;
+}
+
+
+/* FAvfMediaTracks implementation
+ *****************************************************************************/
+
+void FAvfMediaTracks::InitializeAudioSink()
+{
+	if ((AudioSink == nullptr) || (SelectedAudioTrack == INDEX_NONE))
+	{
+		return;
+	}
+
+	AudioSink->InitializeAudioSink(
+		GetAudioTrackChannels(SelectedAudioTrack),
+		GetAudioTrackSampleRate(SelectedAudioTrack)
+	);
+}
+
+
+void FAvfMediaTracks::InitializeCaptionSink()
+{
+	if ((CaptionSink == nullptr) || (SelectedCaptionTrack == INDEX_NONE))
+	{
+		return;
+	}
+
+	CaptionSink->InitializeStringSink();
+}
+
+
+void FAvfMediaTracks::InitializeVideoSink()
+{
+	if ((VideoSink == nullptr) || (SelectedVideoTrack == INDEX_NONE))
+	{
+		return;
+	}
+
+	VideoSink->InitializeTextureSink(
+		GetVideoTrackDimensions(SelectedVideoTrack),
+		EMediaTextureSinkFormat::CharBGRA,
+		EMediaTextureSinkMode::Unbuffered
+	);
+}
diff --git a/Engine/Plugins/Media/AvfMedia/Source/AvfMedia/Private/Player/AvfMediaTracks.h b/Engine/Plugins/Media/AvfMedia/Source/AvfMedia/Private/Player/AvfMediaTracks.h
new file mode 100644
--- /dev/null
+++ b/Engine/Plugins/Media/AvfMedia/Source/AvfMedia/Private/Player/AvfMediaTracks.h
@@ -0,0 +1,138 @@
+// Copyright 1998-2016 Epic Games, Inc. All Rights Reserved.
+
+#pragma once
+
+#include "IMediaOutput.h"
+#include "IMediaTracks.h"
+
+
+class FAvfMediaTracks
+	: public IMediaOutput
+	, public IMediaTracks
+{
+	enum ESyncStatus
+	{
+		Default,    // Starting state
+		Ahead,      // Frame is ahead of playback cursor.
+		Behind,     // Frame is behind playback cursor.
+		Ready,      // Frame is within tolerance of playback cursor.
+	};
+
+	struct FTrack
+	{
+		AVAssetTrack* AssetTrack;
+		FText DisplayName;
+		bool Loaded;
+		FString Name;
+		AVAssetReaderTrackOutput* Output;
+		AVAssetReader* Reader;
+		ESyncStatus SyncStatus;
+
+		// @todo trepka: do we have to 'release' any of these in the destructor?
+	};
+
+public:
+
+	/** Default constructor. */
+	FAvfMediaTracks();
+
+	/** Virtual destructor. */
+	virtual ~FAvfMediaTracks();
+
+public:
+
+	/** Discard any pending data on the media sinks. */
+	void Flush();
+
+	/**
+	 * Initialize the track collection.
+	 *
+	 * @param PlayerItem The player item containing the track information.
+	 */
+	void Initialize(AVPlayerItem& InPlayerItem);
+
+	/** Reset the stream collection. */
+	void Reset();
+
+	/**
+	 * Notify tracks that playback was paused or resumed.
+	 *
+	 * @param Paused true if playback was paused, false if resumed.
+	 */
+	void SetPaused(bool Paused);
+
+public:
+
+	//~ IMediaOutput interface
+
+	virtual void SetAudioSink(IMediaAudioSink* Sink) override;
+	virtual void SetCaptionSink(IMediaStringSink* Sink) override;
+	virtual void SetImageSink(IMediaTextureSink* Sink) override;
+	virtual void SetVideoSink(IMediaTextureSink* Sink) override;
+
+public:
+
+	//~ IMediaTracks interface
+
+	virtual uint32 GetAudioTrackChannels(int32 TrackIndex) const override;
+	virtual uint32 GetAudioTrackSampleRate(int32 TrackIndex) const override;
+	virtual int32 GetNumTracks(EMediaTrackType TrackType) const override;
+	virtual int32 GetSelectedTrack(EMediaTrackType TrackType) const override;
+	virtual FText GetTrackDisplayName(EMediaTrackType TrackType, int32 TrackIndex) const override;
+	virtual FString GetTrackLanguage(EMediaTrackType TrackType, int32 TrackIndex) const override;
+	virtual FString GetTrackName(EMediaTrackType TrackType, int32 TrackIndex) const override;
+	virtual uint32 GetVideoTrackBitRate(int32 TrackIndex) const override;
+	virtual FIntPoint GetVideoTrackDimensions(int32 TrackIndex) const override;
+	virtual float GetVideoTrackFrameRate(int32 TrackIndex) const override;
+	virtual bool SelectTrack(EMediaTrackType TrackType, int32 TrackIndex) override;
+
+protected:
+
+	/** Initialize the current audio sink. */
+	void InitializeAudioSink();
+
+	/** Initialize the current caption sink. */
+	void InitializeCaptionSink();
+
+	/** Initialize the current video sink. */
+	void InitializeVideoSink();
+
+private:
+
+	/** The currently used audio sink. */
+	IMediaAudioSink* AudioSink;
+
+	/** The currently used caption sink. */
+	IMediaStringSink* CaptionSink;
+
+	/** The currently used video sink. */
+	IMediaTextureSink* VideoSink;
+
+private:
+
+	/** The available audio tracks. */
+	TArray<FTrack> AudioTracks;
+
+	/** The available caption tracks. */
+	TArray<FTrack> CaptionTracks;
+
+	/** The available video tracks. */
+	TArray<FTrack> VideoTracks;
+
+private:
+
+	/** Synchronizes write access to track arrays, selections & sinks. */
+	mutable FCriticalSection CriticalSection;
+
+	/** The player item containing the track information. */
+	AVPlayerItem* PlayerItem;
+
+	/** Index of the selected audio track. */
+	int32 SelectedAudioTrack;
+
+	/** Index of the selected caption track. */
+	int32 SelectedCaptionTrack;
+
+	/** Index of the selected video track. */
+	int32 SelectedVideoTrack;
+};
