Copying //UE4/Dev-Sequencer to //UE4/Main (Source: //UE4/Dev-Sequencer @ 3617642)

diff --git a/Engine/Plugins/Media/AvfMedia/Source/AvfMedia/Private/Player/AvfMediaTracks.cpp b/Engine/Plugins/Media/AvfMedia/Source/AvfMedia/Private/Player/AvfMediaTracks.cpp
--- a/Engine/Plugins/Media/AvfMedia/Source/AvfMedia/Private/Player/AvfMediaTracks.cpp
+++ b/Engine/Plugins/Media/AvfMedia/Source/AvfMedia/Private/Player/AvfMediaTracks.cpp
@@ -1,20 +1,28 @@
 // Copyright 1998-2017 Epic Games, Inc. All Rights Reserved.
 
 #include "AvfMediaTracks.h"
+#include "AvfMediaPrivate.h"
+
 #include "Containers/ResourceArray.h"
+#include "MediaSamples.h"
 #include "Misc/ScopeLock.h"
 
 #if WITH_ENGINE
-#include "RenderingThread.h"
-#include "RHI.h"
-#include "MediaShaders.h"
-#include "StaticBoundShaderState.h"
-#include "Misc/ScopeLock.h"
-#include "PipelineStateCache.h"
-#include "RHIStaticStates.h"
-#include "Class.h"
+	#include "RenderingThread.h"
 #endif
 
+#include "AvfMediaAudioSample.h"
+#include "AvfMediaOverlaySample.h"
+#include "AvfMediaVideoSampler.h"
+#include "AvfMediaUtils.h"
+
+#import <AudioToolbox/AudioToolbox.h>
+
+#define AUDIO_PLAYBACK_VIA_ENGINE (PLATFORM_MAC && 0)
+
+/* FAVPlayerItemLegibleOutputPushDelegate
+ *****************************************************************************/
+
 @interface FAVPlayerItemLegibleOutputPushDelegate : NSObject<AVPlayerItemLegibleOutputPushDelegate>
 {
 	FAvfMediaTracks* Tracks;
@@ -37,134 +45,28 @@
 
 - (void)legibleOutput:(AVPlayerItemLegibleOutput *)Output didOutputAttributedStrings:(NSArray<NSAttributedString *> *)Strings nativeSampleBuffers:(NSArray *)NativeSamples forItemTime:(CMTime)ItemTime
 {
-	Tracks->HandleSubtitleCaptionStrings(Output, Strings, NativeSamples, ItemTime);
+	Tracks->ProcessCaptions(Output, Strings, NativeSamples, ItemTime);
 }
 
 @end
 
-/**
- * Passes a CV*TextureRef or CVPixelBufferRef through to the RHI to wrap in an RHI texture without traversing system memory.
- */
-class FAvfTexture2DResourceWrapper : public FResourceBulkDataInterface
-{
-public:
-	FAvfTexture2DResourceWrapper(CFTypeRef InImageBuffer)
-	: ImageBuffer(InImageBuffer)
-	{
-		check(ImageBuffer);
-		CFRetain(ImageBuffer);
-	}
-	
-	/**
-	 * @return ptr to the resource memory which has been preallocated
-	 */
-	virtual const void* GetResourceBulkData() const override
-	{
-		return ImageBuffer;
-	}
-	
-	/**
-	 * @return size of resource memory
-	 */
-	virtual uint32 GetResourceBulkDataSize() const override
-	{
-		return 0;
-	}
-	
-	/**
-	 * Free memory after it has been used to initialize RHI resource
-	 */
-	virtual void Discard() override
-	{
-		delete this;
-	}
-	
-	virtual EBulkDataType GetResourceType() const override
-	{
-		return EBulkDataType::MediaTexture;
-	}
-	
-	virtual ~FAvfTexture2DResourceWrapper()
-	{
-		CFRelease(ImageBuffer);
-		ImageBuffer = nullptr;
-	}
-	
-	CFTypeRef ImageBuffer;
-};
-
-/**
- * Allows for direct GPU mem allocation for texture resource from a CVImageBufferRef's system memory backing store.
- */
-class FAvfTexture2DResourceMem : public FResourceBulkDataInterface
-{
-public:
-	FAvfTexture2DResourceMem(CVImageBufferRef InImageBuffer)
-	: ImageBuffer(InImageBuffer)
-	{
-		check(ImageBuffer);
-		CFRetain(ImageBuffer);
-	}
-	
-	/**
-	 * @return ptr to the resource memory which has been preallocated
-	 */
-	virtual const void* GetResourceBulkData() const override
-	{
-		CVPixelBufferLockBaseAddress(ImageBuffer, kCVPixelBufferLock_ReadOnly);
-		return CVPixelBufferGetBaseAddress(ImageBuffer);
-	}
-	
-	/**
-	 * @return size of resource memory
-	 */
-	virtual uint32 GetResourceBulkDataSize() const override
-	{
-		int32 Pitch = CVPixelBufferGetBytesPerRow(ImageBuffer);
-		int32 Height = CVPixelBufferGetHeight(ImageBuffer);
-		uint32 Size = (Pitch * Height);
-		return Size;
-	}
-	
-	/**
-	 * Free memory after it has been used to initialize RHI resource
-	 */
-	virtual void Discard() override
-	{
-		CVPixelBufferUnlockBaseAddress(ImageBuffer, kCVPixelBufferLock_ReadOnly);
-		delete this;
-	}
-	
-	virtual ~FAvfTexture2DResourceMem()
-	{
-		CFRelease(ImageBuffer);
-		ImageBuffer = nullptr;
-	}
-	
-	CVImageBufferRef ImageBuffer;
-};
 
 /* FAvfMediaTracks structors
  *****************************************************************************/
 
-FAvfMediaTracks::FAvfMediaTracks()
-	: AudioSink(nullptr)
-	, OverlaySink(nullptr)
-	, VideoSink(nullptr)
+FAvfMediaTracks::FAvfMediaTracks(FMediaSamples& InSamples)
+	: AudioPaused(false)
+	, AudioSamplePool(new FAvfMediaAudioSamplePool)
 	, PlayerItem(nullptr)
-	, VideoSampler(nullptr)
+	, Samples(InSamples)
+	, SeekTime(-1.0)
 	, SelectedAudioTrack(INDEX_NONE)
 	, SelectedCaptionTrack(INDEX_NONE)
 	, SelectedVideoTrack(INDEX_NONE)
-	, bAudioPaused(false)
-	, SeekTime(-1.0)
-	, bZoomed(false)
+	, Zoomed(false)
 {
-	LastAudioSample = kCMTimeZero;
-	ENQUEUE_UNIQUE_RENDER_COMMAND_ONEPARAMETER(InitialiseAvfVideoSampler, FAvfMediaTracks*, This, this,
-	{
-		This->VideoSampler = new FAvfVideoSampler();
-	});
+	LastAudioSampleTime = kCMTimeZero;
+	VideoSampler = MakeShared<FAvfMediaVideoSampler, ESPMode::ThreadSafe>(Samples);
 }
 
 
@@ -172,93 +74,55 @@ FAvfMediaTracks::~FAvfMediaTracks()
 {
 	Reset();
 
-	ENQUEUE_UNIQUE_RENDER_COMMAND_ONEPARAMETER(UninitialiseAvfVideoSampler, FAvfMediaTracks*, This, this,
-	{
-		delete This->VideoSampler;
-	});
-
-	FlushRenderingCommands();
+	delete AudioSamplePool;
+	AudioSamplePool = nullptr;
 }
 
 
 /* FAvfMediaTracks interface
  *****************************************************************************/
 
-void FAvfMediaTracks::Flush()
+void FAvfMediaTracks::AppendStats(FString &OutStats) const
 {
 	FScopeLock Lock(&CriticalSection);
 
-	if (AudioSink != nullptr)
-	{
-		AudioSink->FlushAudioSink();
-	}
-}
+	// audio tracks
+	OutStats += TEXT("Audio Tracks\n");
 
-static FString MediaTypeToString(NSString* MediaType)
-{
-	if ([MediaType isEqualToString:AVMediaTypeAudio])
-	{
-		return TEXT("Audio");
-	}
-	else if ([MediaType isEqualToString:AVMediaTypeClosedCaption])
-	{
-		return TEXT("Closed Caption");
-	}
-	else if ([MediaType isEqualToString:AVMediaTypeSubtitle])
+	if (AudioTracks.Num() == 0)
 	{
-		return TEXT("Subtitle");
+		OutStats += TEXT("    none\n");
 	}
-	else if ([MediaType isEqualToString:AVMediaTypeText])
+	else
 	{
-		return TEXT("Text");
+		for (uint32 TrackIndex = 0; TrackIndex < AudioTracks.Num(); TrackIndex++)
+		{
+			const FTrack& Track = AudioTracks[TrackIndex];
+
+			OutStats += FString::Printf(TEXT("    %s\n"), *GetTrackDisplayName(EMediaTrackType::Audio, TrackIndex).ToString());
+			OutStats += FString::Printf(TEXT("        Not implemented yet"));
+		}
 	}
-	else if ([MediaType isEqualToString:AVMediaTypeTimecode])
+
+	// video tracks
+	OutStats += TEXT("Video Tracks\n");
+
+	if (VideoTracks.Num() == 0)
 	{
-		return TEXT("Timecode (unsupported)");
+		OutStats += TEXT("    none\n");
 	}
-	else if ([MediaType isEqualToString:AVMediaTypeVideo])
+	else
 	{
-		return TEXT("Video");
+		for (uint32 TrackIndex = 0; TrackIndex < VideoTracks.Num(); TrackIndex++)
+		{
+			const FTrack& Track = VideoTracks[TrackIndex];
+
+			OutStats += FString::Printf(TEXT("    %s\n"), *GetTrackDisplayName(EMediaTrackType::Video, TrackIndex).ToString());
+			OutStats += FString::Printf(TEXT("        BitRate: %i\n"), [Track.AssetTrack estimatedDataRate]);
+		}
 	}
-	return TEXT("Unknown");
 }
 
-static FString CodecTypeToString(CMVideoCodecType CodecType)
-{
-	if (CodecType == kCMVideoCodecType_422YpCbCr8) return TEXT("422YpCbCr8");
-	if (CodecType == kCMVideoCodecType_Animation) return TEXT("Animation");
-	if (CodecType == kCMVideoCodecType_Cinepak) return TEXT("Cinepak");
-	if (CodecType == kCMVideoCodecType_JPEG) return TEXT("JPEG");
-	if (CodecType == kCMVideoCodecType_JPEG_OpenDML) return TEXT("JPEG OpenDML");
-	if (CodecType == kCMVideoCodecType_SorensonVideo) return TEXT("Sorenson Video");
-	if (CodecType == kCMVideoCodecType_SorensonVideo3) return TEXT("Sorenson Video 3");
-	if (CodecType == kCMVideoCodecType_H263) return TEXT("H263");
-	if (CodecType == kCMVideoCodecType_H264) return TEXT("H264");
-	if (CodecType == kCMVideoCodecType_HEVC) return TEXT("HEVC");
-	if (CodecType == kCMVideoCodecType_MPEG4Video) return TEXT("MPEG4 Video");
-	if (CodecType == kCMVideoCodecType_MPEG2Video) return TEXT("MPEG2 Video");
-	if (CodecType == kCMVideoCodecType_MPEG1Video) return TEXT("MPEG1 Video");
-
-	if (CodecType == kCMVideoCodecType_DVCNTSC) return TEXT("DVC NTSC");
-	if (CodecType == kCMVideoCodecType_DVCPAL) return TEXT("DVC PAL");
-	if (CodecType == kCMVideoCodecType_DVCProPAL) return TEXT("DVCPro PAL");
-	if (CodecType == kCMVideoCodecType_DVCPro50NTSC) return TEXT("DVCPro50 NTSC");
-	if (CodecType == kCMVideoCodecType_DVCPro50PAL) return TEXT("DVCPro50 PAL");
-	if (CodecType == kCMVideoCodecType_DVCPROHD720p60) return TEXT("DVCPRO HD 720p 60");
-	if (CodecType == kCMVideoCodecType_DVCPROHD720p50) return TEXT("DVCPRO HD 720p 50");
-	if (CodecType == kCMVideoCodecType_DVCPROHD1080i60) return TEXT("DVCPRO HD 1080i 60");
-	if (CodecType == kCMVideoCodecType_DVCPROHD1080i50) return TEXT("DVCPRO HD 1080i 50");
-	if (CodecType == kCMVideoCodecType_DVCPROHD1080p30) return TEXT("DVCPRO HD 1080p 30");
-	if (CodecType == kCMVideoCodecType_DVCPROHD1080p25) return TEXT("DVCPRO HD 1080p 25");
-
-	if (CodecType == kCMVideoCodecType_AppleProRes4444) return TEXT("Apple ProRes 4444");
-	if (CodecType == kCMVideoCodecType_AppleProRes422HQ) return TEXT("Apple ProRes 422 HQ");
-	if (CodecType == kCMVideoCodecType_AppleProRes422) return TEXT("Apple ProRes 422");
-	if (CodecType == kCMVideoCodecType_AppleProRes422LT) return TEXT("Apple ProRes 422 LT");
-	if (CodecType == kCMVideoCodecType_AppleProRes422Proxy) return TEXT("Apple ProRes 422 Proxy");
-
-	return TEXT("Unknown");
-}
 
 void FAvfMediaTracks::Initialize(AVPlayerItem* InPlayerItem, FString& OutInfo)
 {
@@ -270,7 +134,6 @@ void FAvfMediaTracks::Initialize(AVPlayerItem* InPlayerItem, FString& OutInfo)
 
 	// initialize tracks
 	NSArray* PlayerTracks = PlayerItem.tracks;
-
 	NSError* Error = nil;
 
 	int32 StreamIndex = 0;
@@ -283,7 +146,7 @@ void FAvfMediaTracks::Initialize(AVPlayerItem* InPlayerItem, FString& OutInfo)
 		NSString* MediaType = AssetTrack.mediaType;
 
 		OutInfo += FString::Printf(TEXT("Stream %i\n"), StreamIndex);
-		OutInfo += FString::Printf(TEXT("    Type: %s\n"), *MediaTypeToString(MediaType));
+		OutInfo += FString::Printf(TEXT("    Type: %s\n"), *AvfMedia::MediaTypeToString(MediaType));
 
 		if ([MediaType isEqualToString:AVMediaTypeAudio])
 		{
@@ -293,7 +156,7 @@ void FAvfMediaTracks::Initialize(AVPlayerItem* InPlayerItem, FString& OutInfo)
 			Track->Name = FString::Printf(TEXT("Audio Track %i"), TrackIndex);
 			Track->Converter = nullptr;
 			
-#if PLATFORM_MAC
+#if AUDIO_PLAYBACK_VIA_ENGINE
 			// create asset reader
 			AVAssetReader* Reader = [[AVAssetReader alloc] initWithAsset: [AssetTrack asset] error:&Error];
 			
@@ -328,6 +191,7 @@ void FAvfMediaTracks::Initialize(AVPlayerItem* InPlayerItem, FString& OutInfo)
 
 			CMFormatDescriptionRef DescRef = (CMFormatDescriptionRef)[AssetTrack.formatDescriptions objectAtIndex:0];
 			const AudioStreamBasicDescription* Desc = CMAudioFormatDescriptionGetStreamBasicDescription(DescRef);
+
 			if (Desc)
 			{
 				OutInfo += FString::Printf(TEXT("    Channels: %u\n"), Desc->mChannelsPerFrame);
@@ -393,12 +257,13 @@ void FAvfMediaTracks::Initialize(AVPlayerItem* InPlayerItem, FString& OutInfo)
 			{
 				[OutputSettings setObject:[NSNumber numberWithBool:YES] forKey:(NSString*)kCVPixelBufferMetalCompatibilityKey];
 			}
+
 #if PLATFORM_MAC
 			[OutputSettings setObject:[NSNumber numberWithBool:YES] forKey:(NSString*)kCVPixelBufferOpenGLCompatibilityKey];
 #else
 			[OutputSettings setObject:[NSNumber numberWithBool:YES] forKey:(NSString*)kCVPixelBufferOpenGLESCompatibilityKey];
 #endif
-#endif
+#endif //WITH_ENGINE
 
 			// Use unaligned rows
 			[OutputSettings setObject:[NSNumber numberWithInteger:1] forKey:(NSString*)kCVPixelBufferBytesPerRowAlignmentKey];
@@ -420,462 +285,370 @@ void FAvfMediaTracks::Initialize(AVPlayerItem* InPlayerItem, FString& OutInfo)
 
 			CMFormatDescriptionRef DescRef = (CMFormatDescriptionRef)[AssetTrack.formatDescriptions objectAtIndex:0];
 			CMVideoCodecType CodecType = CMFormatDescriptionGetMediaSubType(DescRef);
-			OutInfo += FString::Printf(TEXT("    Codec: %s\n"), *CodecTypeToString(CodecType));
+			OutInfo += FString::Printf(TEXT("    Codec: %s\n"), *AvfMedia::CodecTypeToString(CodecType));
 			OutInfo += FString::Printf(TEXT("    Dimensions: %i x %i\n"), (int32)AssetTrack.naturalSize.width, (int32)AssetTrack.naturalSize.height);
 			OutInfo += FString::Printf(TEXT("    Frame Rate: %g fps\n"), AssetTrack.nominalFrameRate);
 			OutInfo += FString::Printf(TEXT("    BitRate: %i\n"), (int32)AssetTrack.estimatedDataRate);
 		}
 
 		OutInfo += TEXT("\n");
-		
-		PlayerTrack.enabled = NO;
 
-		int32 TrackStreamIndex = StreamIndex++;
+		PlayerTrack.enabled = NO;
 
-		if (Track == nullptr)
+		if (Track != nullptr)
 		{
-			continue;
+			Track->AssetTrack = AssetTrack;
+			Track->DisplayName = FText::FromString(Track->Name);
+			Track->StreamIndex = StreamIndex;
 		}
 
-		Track->AssetTrack = AssetTrack;
-		Track->DisplayName = FText::FromString(Track->Name);
-		Track->StreamIndex = TrackStreamIndex;
+		++StreamIndex;
 	}
 }
 
 
-void FAvfMediaTracks::Reset()
+void FAvfMediaTracks::ProcessAudio()
 {
 	FScopeLock Lock(&CriticalSection);
 
-	if (VideoSampler)
+#if AUDIO_PLAYBACK_VIA_ENGINE
+	if (!AudioTracks.IsValidIndex(SelectedAudioTrack) || !PlayerItem || PlayerItem.status != AVPlayerItemStatusReadyToPlay || CMTimeCompare(PlayerItem.duration, CMTimeMakeWithSeconds(0.0, 1000)) < 1)
 	{
-		if (VideoSampler->IsTickable())
-		{
-			ENQUEUE_UNIQUE_RENDER_COMMAND_ONEPARAMETER(InitialiseAvfVideoSampler, FAvfMediaTracks*, This, this,
-			{
-				This->VideoSampler->SetTrack(nullptr, nil);
-			});
-		}
+		return;
 	}
 
-	SelectedAudioTrack = INDEX_NONE;
-	SelectedCaptionTrack = INDEX_NONE;
-	SelectedVideoTrack = INDEX_NONE;
+	CMTime CurrentTime = [PlayerItem currentTime];
 
-	if (AudioSink != nullptr)
-	{
-		AudioSink->ShutdownAudioSink();
-		AudioSink = nullptr;
-	}
+	ESyncStatus Sync = Default;
 
-	if (OverlaySink != nullptr)
-	{
-		OverlaySink->ShutdownOverlaySink();
-		OverlaySink = nullptr;
-	}
+	AVAssetReaderTrackOutput* AudioReaderOutput = (AVAssetReaderTrackOutput*)AudioTracks[SelectedAudioTrack].Output;
+	check(AudioReaderOutput);
 
-	if (VideoSink != nullptr)
+	while (Sync < Ready)
 	{
-		VideoSink->ShutdownTextureSink();
-		VideoSink = nullptr;
-	}
+		float Delta = CMTimeGetSeconds(LastAudioSampleTime) - CMTimeGetSeconds(CurrentTime);
 
-	for (FTrack& Track : AudioTracks)
-	{
-		if (Track.Converter)
+		if ((Delta <= 1.0f) && (Delta > 0.0f))
 		{
-			OSStatus CAErr = AudioConverterDispose(Track.Converter);
-			check(CAErr == noErr);
-		}
-		[Track.Output release];
-		[Track.Reader release];
-	}
+			Sync = Ready;
 
-	for (FTrack& Track : CaptionTracks)
-	{
-		AVPlayerItemLegibleOutput* Output = (AVPlayerItemLegibleOutput*)Track.Output;
-		[Output.delegate release];
-		[Track.Output release];
-		[Track.Reader release];
-	}
-
-	for (FTrack& Track : VideoTracks)
-	{
-		[Track.Output release];
-		[Track.Reader release];
-	}
-
-	AudioTracks.Empty();
-	CaptionTracks.Empty();
-	VideoTracks.Empty();
-	
-	LastAudioSample = kCMTimeZero;
-	bAudioPaused = false;
-	SeekTime = -1.0;
-	bZoomed = false;
-}
+			break;
+		}
 
+		Sync = Behind;
 
-void FAvfMediaTracks::SetRate(float Rate)
-{
-	FScopeLock Lock(&CriticalSection);
+		CMSampleBufferRef LatestSamples = [AudioReaderOutput copyNextSampleBuffer];
 
-	if (AudioSink != nullptr)
-	{
-		// Can only play sensible audio at full rate forward - when seeking, scrubbing or reversing we can't supply
-		// the correct samples.
-		bool bNearOne = FMath::IsNearlyEqual(Rate, 1.0f);
-		
-		bool bWasPaused = bAudioPaused;
-		
-		bAudioPaused = !bNearOne;
-		
-		if (bAudioPaused)
-		{
-#if PLATFORM_MAC
-			if (!FMath::IsNearlyZero(Rate))
-			{
-				AudioSink->FlushAudioSink();
-			}
-#endif
-			AudioSink->PauseAudioSink();
-		}
-		else
+		if (!LatestSamples)
 		{
-#if PLATFORM_MAC
-			if (bZoomed)
-			{
-				CMTime CurrentTime = [PlayerItem currentTime];
-				FTimespan Time = FTimespan::FromSeconds(CMTimeGetSeconds(CurrentTime));
-				Seek(Time);
-			}
-#endif
-			AudioSink->ResumeAudioSink();
+			break;
 		}
-		
-		bZoomed = !bNearOne;
-	}
-}
 
+		CMTime FrameTimeStamp = CMSampleBufferGetPresentationTimeStamp(LatestSamples);
+		CMTime Duration = CMSampleBufferGetOutputDuration(LatestSamples);
+		CMTime FinalTimeStamp = CMTimeAdd(FrameTimeStamp, Duration);
+		CMTime Seek = CMTimeMakeWithSeconds(SeekTime, 1000);
 
-void FAvfMediaTracks::Seek(const FTimespan& Time)
-{
-#if PLATFORM_MAC
-	FScopeLock Lock(&CriticalSection);
-	
-	if (AudioSink && SelectedAudioTrack != INDEX_NONE)
-	{
-		double LastSample = CMTimeGetSeconds(LastAudioSample);
-		SeekTime = Time.GetTotalSeconds();
-		
-		if (SeekTime < LastSample)
+		if ((SeekTime < 0.0f) || (CMTimeCompare(Seek, FrameTimeStamp) >= 0 && CMTimeCompare(Seek, FinalTimeStamp) < 0))
 		{
-			AVAssetReaderTrackOutput* AudioReaderOutput = (AVAssetReaderTrackOutput*)AudioTracks[SelectedAudioTrack].Output;
-			check(AudioReaderOutput);
-			
-			LastAudioSample = CMTimeMakeWithSeconds(0, 1000);
-		
-			CMSampleBufferRef LatestSamples = nullptr;
-			while ((LatestSamples = [AudioReaderOutput copyNextSampleBuffer]))
+			SeekTime = -1.0f;
+
+			CMItemCount NumSamples = CMSampleBufferGetNumSamples(LatestSamples);
+			CMFormatDescriptionRef Format = CMSampleBufferGetFormatDescription(LatestSamples);
+			check(Format);
+
+			const AudioStreamBasicDescription* ASBD = CMAudioFormatDescriptionGetStreamBasicDescription(Format);
+			check(ASBD);
+
+			CMBlockBufferRef Buffer = CMSampleBufferGetDataBuffer(LatestSamples);
+			check(Buffer);
+
+			// create & add sample to queue
+			uint32 InputLength = CMBlockBufferGetDataLength(Buffer);
+
+			if (InputLength)
 			{
-				if (LatestSamples)
+				uint32 OutputLength = (NumSamples * TargetDesc.mBytesPerPacket);
+				FTrack& AudioTrack = AudioTracks[SelectedAudioTrack];
+
+				auto AudioSample = MakeShared<FAvfMediaAudioSample, ESPMode::ThreadSafe>();
+
+				if (!AudioSample->Initialize(
+					OutputLength,
+					TargetDesc.mFramesPerPacket,
+					TargetDesc.mChannelsPerFrame,
+					TargetDesc.mSampleRate,
+					FTimespan::FromSeconds(CMTimeGetSeconds(FrameTimeStamp))))
 				{
-					CFRelease(LatestSamples);
+					continue;
 				}
+
+				if (FMemory::Memcmp(ASBD, &TargetDesc, sizeof(AudioStreamBasicDescription)) != 0)
+				{
+					// conversion to 16-bit PCM required
+					uint8* Data = new uint8[InputLength];
+
+					OSErr Err = CMBlockBufferCopyDataBytes(Buffer, 0, InputLength, Data);
+					check(Err == noErr);
+
+					if (AudioTrack.Converter == nullptr)
+					{
+						OSStatus CAErr = AudioConverterNew(ASBD, &TargetDesc, &AudioTrack.Converter);
+						check(CAErr == noErr);
+					}
+
+					OSStatus CAErr = AudioConverterConvertBuffer(AudioTrack.Converter, InputLength, Data, &OutputLength, AudioSample->GetMutableBuffer());
+					check(CAErr == noErr);
+
+					delete[] Data;
+				}
+				else
+				{
+					// data is already in 16-bit PCM
+					OSErr Err = CMBlockBufferCopyDataBytes(Buffer, 0, OutputLength, AudioSample->GetMutableBuffer());
+					check(Err == noErr);
+				}
+
+				Samples.AddAudio(AudioSample);
+				LastAudioSampleTime = FinalTimeStamp;
 			}
-			
-			CMTime Start = CMTimeMakeWithSeconds(SeekTime, 1000);
-			CMTime Duration = PlayerItem.asset.duration;
-			Duration = CMTimeSubtract(Duration, Start);
-			
-			CMTimeRange TimeRange = CMTimeRangeMake(Start, Duration);
-			NSValue* Value = [NSValue valueWithBytes:&TimeRange objCType:@encode(CMTimeRange)];
-			NSArray* Array = @[ Value ];
-			
-			[AudioReaderOutput resetForReadingTimeRanges:Array];
 		}
-		
-		AudioSink->FlushAudioSink();
+
+		CFRelease(LatestSamples);
 	}
 #endif
 }
 
 
-void FAvfMediaTracks::AppendStats(FString &OutStats) const
+void FAvfMediaTracks::ProcessCaptions(AVPlayerItemLegibleOutput* Output, NSArray<NSAttributedString*>* Strings, NSArray* NativeSamples, CMTime ItemTime)
 {
-	FScopeLock Lock(&CriticalSection);
-
-	// audio tracks
-	OutStats += TEXT("Audio Tracks\n");
-	
-	if (AudioTracks.Num() == 0)
+	if (SelectedCaptionTrack != INDEX_NONE)
 	{
-		OutStats += TEXT("    none\n");
+		return;
 	}
-	else
+
+	FScopeLock Lock(&CriticalSection);
+		
+	NSDictionary* DocumentAttributes = @{NSDocumentTypeDocumentAttribute:NSPlainTextDocumentType};
+	FTimespan DisplayTime = FTimespan::FromSeconds(CMTimeGetSeconds(ItemTime));
+	
+	FString OutputString;
+	bool bFirst = true;
+
+	for (NSAttributedString* String in Strings)
 	{
-		for (uint32 i = 0; i < GetNumTracks(EMediaTrackType::Audio); i++)
+		if (!String)
 		{
-			OutStats += FString::Printf(TEXT("    %s\n"), *GetTrackDisplayName(EMediaTrackType::Audio, i).ToString());
-			OutStats += FString::Printf(TEXT("        Channels: %i\n"), GetAudioTrackChannels(i));
-			OutStats += FString::Printf(TEXT("        Sample Rate: %i\n\n"), GetAudioTrackSampleRate(i));
+			continue;
 		}
-	}
 
-	// video tracks
-	OutStats += TEXT("Video Tracks\n");
+		// strip attributes from the string (we don't care for them)
+		NSRange Range = NSMakeRange(0, String.length);
+		NSData* Data = [String dataFromRange:Range documentAttributes:DocumentAttributes error:NULL];
+		NSString* Result = [[NSString alloc] initWithData:Data encoding:NSUTF8StringEncoding];
+				
+		// append the string
+		if (!bFirst)
+		{
+			OutputString += TEXT("\n");
+		}
 
-	if (VideoTracks.Num() == 0)
+		bFirst = false;
+		OutputString += FString(Result);
+	}
+
+	if (OutputString.IsEmpty())
 	{
-		OutStats += TEXT("    none\n");
+		return;
 	}
-	else
+
+	// create & add sample to queue
+	auto OverlaySample = MakeShared<FAvfMediaOverlaySample, ESPMode::ThreadSafe>();
+
+	if (OverlaySample->Initialize(OutputString, DisplayTime))
 	{
-		for (uint32 i = 0; i < GetNumTracks(EMediaTrackType::Video); i++)
-		{
-			OutStats += FString::Printf(TEXT("    %s\n"), *GetTrackDisplayName(EMediaTrackType::Video, i).ToString());
-			OutStats += FString::Printf(TEXT("        BitRate: %i\n"), GetVideoTrackBitRate(i));
-			OutStats += FString::Printf(TEXT("        Frame Rate: %f\n\n"), GetVideoTrackFrameRate(i));
-		}
+		Samples.AddCaption(OverlaySample);
 	}
 }
 
 
-bool FAvfMediaTracks::Tick(float DeltaTime)
+void FAvfMediaTracks::ProcessVideo()
 {
-	FScopeLock Lock(&CriticalSection);
-
-#if PLATFORM_MAC
-	if (AudioSink && SelectedAudioTrack != INDEX_NONE)
+	struct FVideoSamplerTickParams
 	{
-		CMTime CurrentTime = [PlayerItem currentTime];
-		
-		ESyncStatus Sync = Default;
+		TWeakPtr<FAvfMediaVideoSampler, ESPMode::ThreadSafe> VideoSamplerPtr;
+	}
+	VideoSamplerTickParams = { VideoSampler };
 
-		AVAssetReaderTrackOutput* AudioReaderOutput = (AVAssetReaderTrackOutput*)AudioTracks[SelectedAudioTrack].Output;
-		check(AudioReaderOutput);
-		
-		while (Sync < Ready)
+	ENQUEUE_UNIQUE_RENDER_COMMAND_ONEPARAMETER(AvfMediaVideoSamplerTick,
+		FVideoSamplerTickParams, Params, VideoSamplerTickParams,
 		{
-			float Delta = CMTimeGetSeconds(LastAudioSample) - CMTimeGetSeconds(CurrentTime);
-			if (Delta <= 1.0f)
-			{
-				Sync = Behind;
-				
-				CMSampleBufferRef LatestSamples = [AudioReaderOutput copyNextSampleBuffer];
-				if (LatestSamples)
-				{
-					CMTime FrameTimeStamp = CMSampleBufferGetPresentationTimeStamp(LatestSamples);
-					CMTime Duration = CMSampleBufferGetOutputDuration(LatestSamples);
-					
-					CMTime FinalTimeStamp = CMTimeAdd(FrameTimeStamp, Duration);
-					
-					CMTime Seek = CMTimeMakeWithSeconds(SeekTime, 1000);
-					if ((SeekTime < 0.0f) || (CMTimeCompare(Seek, FrameTimeStamp) >= 0 && CMTimeCompare(Seek, FinalTimeStamp) < 0))
-					{
-						SeekTime = -1.0f;
-						
-						CMItemCount NumSamples = CMSampleBufferGetNumSamples(LatestSamples);
-						
-						CMFormatDescriptionRef Format = CMSampleBufferGetFormatDescription(LatestSamples);
-						check(Format);
-						
-						const AudioStreamBasicDescription* ASBD = CMAudioFormatDescriptionGetStreamBasicDescription(Format);
-						check(ASBD);
-						
-						CMBlockBufferRef Buffer = CMSampleBufferGetDataBuffer(LatestSamples);
-						check(Buffer);
-						
-						uint32 Length = CMBlockBufferGetDataLength(Buffer);
-						if (Length)
-						{
-							uint8* Data = new uint8[Length];
-							
-							OSErr Err = CMBlockBufferCopyDataBytes(Buffer, 0, Length, Data);
-							check(Err == noErr);
-							
-							UInt32 OutputLength = (NumSamples * TargetDesc.mBytesPerPacket);
-							if (FMemory::Memcmp(ASBD, &TargetDesc, sizeof(AudioStreamBasicDescription)) != 0)
-							{
-								if (AudioTracks[SelectedAudioTrack].Converter == nullptr)
-								{
-									OSStatus CAErr = AudioConverterNew(ASBD, &TargetDesc, &AudioTracks[SelectedAudioTrack].Converter);
-									check(CAErr == noErr);
-								}
-								
-								uint8* ConvertedData = new uint8[Length];
-								
-								OSStatus CAErr = AudioConverterConvertBuffer(AudioTracks[SelectedAudioTrack].Converter, (NumSamples * ASBD->mBytesPerPacket), Data, &OutputLength, ConvertedData);
-								check(CAErr == noErr);
-								
-								delete [] Data;
-								Data = ConvertedData;
-							}
-							
-							FTimespan DiplayTime = FTimespan::FromSeconds(CMTimeGetSeconds(FrameTimeStamp));
-							AudioSink->PlayAudioSink(Data, OutputLength, DiplayTime);
-							
-							delete [] Data;
-							
-							LastAudioSample = FinalTimeStamp;
-						}
-					}
-					
-					CFRelease(LatestSamples);
-				}
-				else
-				{
-					break;
-				}
-			}
-			else
+			auto PinnedVideoSampler = Params.VideoSamplerPtr.Pin();
+
+			if (PinnedVideoSampler.IsValid())
 			{
-				Sync = Ready;
+				PinnedVideoSampler->Tick();
 			}
-		}
-	}
-#endif
-	
-	return true;
+		});
 }
 
-void FAvfMediaTracks::HandleSubtitleCaptionStrings(AVPlayerItemLegibleOutput* Output, NSArray<NSAttributedString*>* Strings, NSArray* NativeSamples, CMTime ItemTime)
+
+void FAvfMediaTracks::Reset()
 {
-	if (OverlaySink != nil && SelectedCaptionTrack != INDEX_NONE)
+	FScopeLock Lock(&CriticalSection);
+
+	// reset video sampler
+	struct FResetOutputParams
 	{
-		FScopeLock Lock(&CriticalSection);
-		
-		NSDictionary* DocumentAttributes = @{NSDocumentTypeDocumentAttribute:NSPlainTextDocumentType};
-		FTimespan DiplayTime = FTimespan::FromSeconds(CMTimeGetSeconds(ItemTime));
-	
-		FString OutputString;
-		bool bFirst = true;
+		TWeakPtr<FAvfMediaVideoSampler, ESPMode::ThreadSafe> VideoSamplerPtr;
+	}
+	ResetOutputParams = { VideoSampler };
 
-		for (NSAttributedString* String in Strings)
-		{
-			if(String)
+	ENQUEUE_UNIQUE_RENDER_COMMAND_ONEPARAMETER(AvfMediaVideoSamplerResetOutput,
+		FResetOutputParams, Params, ResetOutputParams,
+	    {
+			auto PinnedVideoSampler = Params.VideoSamplerPtr.Pin();
+
+			if (PinnedVideoSampler.IsValid())
 			{
-				// Strip attribtues from the string - we don't care for them.
-				NSRange Range = NSMakeRange(0, String.length);
-				NSData* Data = [String dataFromRange:Range documentAttributes:DocumentAttributes error:NULL];
-				NSString* Result = [[NSString alloc] initWithData:Data encoding:NSUTF8StringEncoding];
-				
-				// Append the string.
-				if (!bFirst)
-				{
-					OutputString += TEXT("\n");
-				}
-				bFirst = false;
-				OutputString += FString(Result);
+				PinnedVideoSampler->SetOutput(nil, 0.0f);
 			}
+		});
+	
+	// reset tracks
+	SelectedAudioTrack = INDEX_NONE;
+	SelectedCaptionTrack = INDEX_NONE;
+	SelectedVideoTrack = INDEX_NONE;
+
+	for (FTrack& Track : AudioTracks)
+	{
+		if (Track.Converter)
+		{
+			OSStatus CAErr = AudioConverterDispose(Track.Converter);
+			check(CAErr == noErr);
 		}
 
-		//OverlaySink->AddOverlaySinkText(FText::FromString(OutputString), DiplayTime);
+		[Track.Output release];
+		[Track.Reader release];
 	}
-}
-
-
 
-/* IMediaOutput interface
- *****************************************************************************/
-
-void FAvfMediaTracks::SetAudioSink(IMediaAudioSink* Sink)
-{
-	if (Sink != AudioSink)
+	for (FTrack& Track : CaptionTracks)
 	{
-		FScopeLock Lock(&CriticalSection);
+		AVPlayerItemLegibleOutput* Output = (AVPlayerItemLegibleOutput*)Track.Output;
 
-		if (AudioSink != nullptr)
-		{
-			AudioSink->ShutdownAudioSink();
-			AudioSink = nullptr;
-		}
+		[Output.delegate release];
+		[Track.Output release];
+		[Track.Reader release];
+	}
 
-		AudioSink = Sink;	
-		InitializeAudioSink();
+	for (FTrack& Track : VideoTracks)
+	{
+		[Track.Output release];
+		[Track.Reader release];
 	}
+
+	AudioTracks.Empty();
+	CaptionTracks.Empty();
+	VideoTracks.Empty();
+	
+	LastAudioSampleTime = kCMTimeZero;
+	AudioPaused = false;
+	SeekTime = -1.0;
+	Zoomed = false;
 }
 
 
-void FAvfMediaTracks::SetMetadataSink(IMediaBinarySink* Sink)
+void FAvfMediaTracks::Seek(const FTimespan& Time)
 {
-	// not supported
-}
+#if AUDIO_PLAYBACK_VIA_ENGINE
+	FScopeLock Lock(&CriticalSection);
 
+	if (SelectedAudioTrack != INDEX_NONE)
+	{
+		return;
+	}
 
-void FAvfMediaTracks::SetOverlaySink(IMediaOverlaySink* Sink)
-{
-	if (Sink != OverlaySink)
+	double LastSampleTime = CMTimeGetSeconds(LastAudioSampleTime);
+	SeekTime = Time.GetTotalSeconds();
+
+	if (SeekTime >= LastSampleTime)
 	{
-		FScopeLock Lock(&CriticalSection);
+		return;
+	}
+
+	AVAssetReaderTrackOutput* AudioReaderOutput = (AVAssetReaderTrackOutput*)AudioTracks[SelectedAudioTrack].Output;
+	check(AudioReaderOutput);
 
-		if (OverlaySink != nullptr)
+	LastAudioSampleTime = CMTimeMakeWithSeconds(0, 1000);
+
+	CMSampleBufferRef LatestSamples = nullptr;
+	while ((LatestSamples = [AudioReaderOutput copyNextSampleBuffer]))
+	{
+		if (LatestSamples)
 		{
-			OverlaySink->ShutdownOverlaySink();
-			OverlaySink = nullptr;
+			CFRelease(LatestSamples);
 		}
-
-		OverlaySink = Sink;
-		InitializeOverlaySink();
 	}
+
+	CMTime Start = CMTimeMakeWithSeconds(SeekTime, 1000);
+	CMTime Duration = PlayerItem.asset.duration;
+	Duration = CMTimeSubtract(Duration, Start);
+
+	CMTimeRange TimeRange = CMTimeRangeMake(Start, Duration);
+	NSValue* Value = [NSValue valueWithBytes : &TimeRange objCType : @encode(CMTimeRange)];
+	NSArray* Array = @[ Value ];
+
+	[AudioReaderOutput resetForReadingTimeRanges : Array];
+#endif
 }
 
 
-void FAvfMediaTracks::SetVideoSink(IMediaTextureSink* Sink)
+void FAvfMediaTracks::SetRate(float Rate)
 {
-	if (Sink != VideoSink)
-	{
-		FScopeLock Lock(&CriticalSection);
-
-		if (VideoSink != nullptr)
-		{
-			VideoSink->ShutdownTextureSink();
-			VideoSink = nullptr;
-		}
+	FScopeLock Lock(&CriticalSection);
 
-		VideoSink = Sink;
-		InitializeVideoSink();
+	// Can only play sensible audio at full rate forward - when seeking, scrubbing or reversing we can't supply
+	// the correct samples.
+	const bool bNearOne = FMath::IsNearlyEqual(Rate, 1.0f);	
+	const bool bWasPaused = AudioPaused;
+		
+	AudioPaused = !bNearOne;
+		
+#if AUDIO_PLAYBACK_VIA_ENGINE
+	if (!AudioPaused && Zoomed)
+	{
+		CMTime CurrentTime = [PlayerItem currentTime];
+		FTimespan Time = FTimespan::FromSeconds(CMTimeGetSeconds(CurrentTime));
+		Seek(Time);
 	}
+#endif
+		
+	Zoomed = !bNearOne;
 }
 
 
 /* IMediaTracks interface
  *****************************************************************************/
 
-uint32 FAvfMediaTracks::GetAudioTrackChannels(int32 TrackIndex) const
+bool FAvfMediaTracks::GetAudioTrackFormat(int32 TrackIndex, int32 FormatIndex, FMediaAudioTrackFormat& OutFormat) const
 {
-	uint32 NumAudioTracks = 0;
-	if (AudioTracks.IsValidIndex(TrackIndex))
+	if ((FormatIndex != 0) || !AudioTracks.IsValidIndex(TrackIndex))
 	{
-		checkf(AudioTracks[TrackIndex].AssetTrack.formatDescriptions.count == 1, TEXT("Can't handle non-uniform audio streams!"));
-		CMFormatDescriptionRef DescRef = (CMFormatDescriptionRef)[AudioTracks[TrackIndex].AssetTrack.formatDescriptions objectAtIndex:0];
-		AudioStreamBasicDescription const* Desc = CMAudioFormatDescriptionGetStreamBasicDescription(DescRef);
-		if (Desc)
-		{
-			NumAudioTracks = Desc->mChannelsPerFrame;
-		}
+		return false;
 	}
-	
-	return NumAudioTracks;
-}
 
+	const FTrack& Track = AudioTracks[TrackIndex];
+	checkf(Track.AssetTrack.formatDescriptions.count == 1, TEXT("Can't handle non-uniform audio streams!"));
 
-uint32 FAvfMediaTracks::GetAudioTrackSampleRate(int32 TrackIndex) const
-{
-	uint32 SampleRate = 0;
-	if (AudioTracks.IsValidIndex(TrackIndex))
-	{
-		checkf(AudioTracks[TrackIndex].AssetTrack.formatDescriptions.count == 1, TEXT("Can't handle non-uniform audio streams!"));
-		CMFormatDescriptionRef DescRef = (CMFormatDescriptionRef)[AudioTracks[TrackIndex].AssetTrack.formatDescriptions objectAtIndex:0];
-		AudioStreamBasicDescription const* Desc = CMAudioFormatDescriptionGetStreamBasicDescription(DescRef);
-		if (Desc)
-		{
-			SampleRate = Desc->mSampleRate;
-		}
-	}
-	
-	return SampleRate;
+	CMFormatDescriptionRef DescRef = (CMFormatDescriptionRef)[AudioTracks[TrackIndex].AssetTrack.formatDescriptions objectAtIndex : 0];
+	AudioStreamBasicDescription const* Desc = CMAudioFormatDescriptionGetStreamBasicDescription(DescRef);
+
+	OutFormat.BitsPerSample = 16;
+	OutFormat.NumChannels = (Desc != nullptr) ? Desc->mChannelsPerFrame : 0;
+	OutFormat.SampleRate = (Desc != nullptr) ? Desc->mSampleRate : 0;
+	OutFormat.TypeName = TEXT("PCM"); // @todo trepka: fix me (should be input type, not output type)
+
+	return true;
 }
 
 
@@ -885,26 +658,38 @@ int32 FAvfMediaTracks::GetNumTracks(EMediaTrackType TrackType) const
 	{
 	case EMediaTrackType::Audio:
 		return AudioTracks.Num();
+
 	case EMediaTrackType::Caption:
 		return CaptionTracks.Num();
+
 	case EMediaTrackType::Video:
 		return VideoTracks.Num();
+
 	default:
 		return 0;
 	}
 }
 
 
+int32 FAvfMediaTracks::GetNumTrackFormats(EMediaTrackType TrackType, int32 TrackIndex) const
+{
+	return ((TrackIndex >= 0) && (TrackIndex < GetNumTracks(TrackType))) ? 1 : 0;
+}
+
+
 int32 FAvfMediaTracks::GetSelectedTrack(EMediaTrackType TrackType) const
 {
 	switch (TrackType)
 	{
 	case EMediaTrackType::Audio:
 		return SelectedAudioTrack;
+
 	case EMediaTrackType::Caption:
 		return SelectedCaptionTrack;
+
 	case EMediaTrackType::Video:
 		return SelectedVideoTrack;
+
 	default:
 		return INDEX_NONE;
 	}
@@ -944,6 +729,12 @@ FText FAvfMediaTracks::GetTrackDisplayName(EMediaTrackType TrackType, int32 Trac
 }
 
 
+int32 FAvfMediaTracks::GetTrackFormat(EMediaTrackType TrackType, int32 TrackIndex) const
+{
+	return (GetSelectedTrack(TrackType) != INDEX_NONE) ? 0 : INDEX_NONE;
+}
+
+
 FString FAvfMediaTracks::GetTrackLanguage(EMediaTrackType TrackType, int32 TrackIndex) const
 {
 	if (TrackType == EMediaTrackType::Audio)
@@ -1005,24 +796,25 @@ FString FAvfMediaTracks::GetTrackName(EMediaTrackType TrackType, int32 TrackInde
 }
 
 
-uint32 FAvfMediaTracks::GetVideoTrackBitRate(int32 TrackIndex) const
+bool FAvfMediaTracks::GetVideoTrackFormat(int32 TrackIndex, int32 FormatIndex, FMediaVideoTrackFormat& OutFormat) const
 {
-	return VideoTracks.IsValidIndex(TrackIndex) ? [VideoTracks[TrackIndex].AssetTrack estimatedDataRate] : 0;
-}
+	if ((FormatIndex != 0) || !VideoTracks.IsValidIndex(TrackIndex))
+	{
+		return false;
+	}
 
+	const FTrack& Track = VideoTracks[TrackIndex];
 
-FIntPoint FAvfMediaTracks::GetVideoTrackDimensions(int32 TrackIndex) const
-{
-	return VideoTracks.IsValidIndex(TrackIndex)
-		? FIntPoint([VideoTracks[TrackIndex].AssetTrack naturalSize].width,
-					[VideoTracks[TrackIndex].AssetTrack naturalSize].height)
-		: FIntPoint::ZeroValue;
-}
+	OutFormat.Dim = FIntPoint(
+		[Track.AssetTrack naturalSize].width,
+		[Track.AssetTrack naturalSize].height
+	);
 
+	OutFormat.FrameRate = [Track.AssetTrack nominalFrameRate];
+	OutFormat.FrameRates = TRange<float>(OutFormat.FrameRate);
+	OutFormat.TypeName = TEXT("BGRA"); // @todo trepka: fix me (should be input format, not output format)
 
-float FAvfMediaTracks::GetVideoTrackFrameRate(int32 TrackIndex) const
-{
-	return VideoTracks.IsValidIndex(TrackIndex) ? 1.0f / [VideoTracks[TrackIndex].AssetTrack nominalFrameRate] : 0.0f;
+	return true;
 }
 
 
@@ -1033,486 +825,223 @@ bool FAvfMediaTracks::SelectTrack(EMediaTrackType TrackType, int32 TrackIndex)
 	switch (TrackType)
 	{
 	case EMediaTrackType::Audio:
-		if ((TrackIndex == INDEX_NONE) || AudioTracks.IsValidIndex(TrackIndex))
+		if (TrackIndex != SelectedAudioTrack)
 		{
-			if (TrackIndex != SelectedAudioTrack)
+			UE_LOG(LogAvfMedia, Verbose, TEXT("Selecting audio track %i instead of %i (%i tracks)."), TrackIndex, SelectedAudioTrack, AudioTracks.Num());
+
+			// disable current track
+			if (SelectedAudioTrack != INDEX_NONE)
 			{
-				if (SelectedAudioTrack != INDEX_NONE)
+				UE_LOG(LogAvfMedia, VeryVerbose, TEXT("Disabling audio track %i"), SelectedAudioTrack);
+
+#if AUDIO_PLAYBACK_VIA_ENGINE
+				AVAssetReaderTrackOutput* AudioReaderOutput = (AVAssetReaderTrackOutput*)AudioTracks[SelectedAudioTrack].Output;
+				check(AudioReaderOutput);
+
+				CMSampleBufferRef LatestSamples = nullptr;
+
+				while ((LatestSamples = [AudioReaderOutput copyNextSampleBuffer]))
 				{
-#if PLATFORM_MAC
-					AVAssetReaderTrackOutput* AudioReaderOutput = (AVAssetReaderTrackOutput*)AudioTracks[SelectedAudioTrack].Output;
-					check(AudioReaderOutput);
-					
-					CMSampleBufferRef LatestSamples = nullptr;
-					while ((LatestSamples = [AudioReaderOutput copyNextSampleBuffer]))
+					if (LatestSamples)
 					{
-						if (LatestSamples)
-						{
-							CFRelease(LatestSamples);
-						}
+						CFRelease(LatestSamples);
 					}
-					
-					CMTimeRange TimeRange = CMTimeRangeMake(kCMTimeZero, PlayerItem.asset.duration);
-					NSValue* Value = [NSValue valueWithBytes:&TimeRange objCType:@encode(CMTimeRange)];
-					NSArray* Array = @[ Value ];
-					
-					[AudioReaderOutput resetForReadingTimeRanges:Array];
+				}
+
+				CMTimeRange TimeRange = CMTimeRangeMake(kCMTimeZero, PlayerItem.asset.duration);
+				NSValue* Value = [NSValue valueWithBytes : &TimeRange objCType : @encode(CMTimeRange)];
+				NSArray* Array = @[ Value ];
+
+				[AudioReaderOutput resetForReadingTimeRanges : Array];
+
 #else
-					AVPlayerItemTrack* PlayerTrack = (AVPlayerItemTrack*)AudioTracks[SelectedAudioTrack].Output;
-					check(PlayerTrack);
-					PlayerTrack.enabled = false;
+				AVPlayerItemTrack* PlayerTrack = (AVPlayerItemTrack*)AudioTracks[SelectedAudioTrack].Output;
+				check(PlayerTrack);
+
+				PlayerTrack.enabled = NO;
 #endif
-					SelectedAudioTrack = INDEX_NONE;
-				}
 
-//				if ((TrackIndex != INDEX_NONE) && SUCCEEDED(PresentationDescriptor->SelectStream(AudioTracks[TrackIndex].StreamIndex)))
-				{
-					SelectedAudioTrack = TrackIndex;
-				}
+				SelectedAudioTrack = INDEX_NONE;
+			}
 
-				if (SelectedAudioTrack == TrackIndex)
+			// enable new track
+			if (TrackIndex != INDEX_NONE)
+			{
+				if (!AudioTracks.IsValidIndex(TrackIndex))
 				{
-					InitializeAudioSink();
+					return false;
 				}
+
+				UE_LOG(LogAvfMedia, VeryVerbose, TEXT("Enabling audio track %i"), TrackIndex);
 			}
 
-			return true;
+			SelectedAudioTrack = TrackIndex;
+
+			// update output
+			if (SelectedAudioTrack != INDEX_NONE)
+			{
+				const FTrack& SelectedTrack = AudioTracks[SelectedAudioTrack];
+				
+#if AUDIO_PLAYBACK_VIA_ENGINE
+				CMFormatDescriptionRef DescRef = (CMFormatDescriptionRef)[SelectedTrack.AssetTrack.formatDescriptions objectAtIndex : 0];
+				AudioStreamBasicDescription const* ASBD = CMAudioFormatDescriptionGetStreamBasicDescription(DescRef);
+
+				TargetDesc.mSampleRate = ASBD->mSampleRate;
+				TargetDesc.mFormatID = kAudioFormatLinearPCM;
+				TargetDesc.mFormatFlags = kAudioFormatFlagIsSignedInteger | kLinearPCMFormatFlagIsPacked;
+				TargetDesc.mBytesPerPacket = ASBD->mChannelsPerFrame * sizeof(int16);
+				TargetDesc.mFramesPerPacket = 1;
+				TargetDesc.mBytesPerFrame = ASBD->mChannelsPerFrame * sizeof(int16);
+				TargetDesc.mChannelsPerFrame = ASBD->mChannelsPerFrame;
+				TargetDesc.mBitsPerChannel = 16;
+				TargetDesc.mReserved = 0;
+#else
+				AVPlayerItemTrack* PlayerTrack = (AVPlayerItemTrack*)SelectedTrack.Output;
+				check(PlayerTrack);
+
+				PlayerTrack.enabled = YES;
+#endif
+			}
 		}
 		break;
 
 	case EMediaTrackType::Caption:
-		if ((TrackIndex == INDEX_NONE) || CaptionTracks.IsValidIndex(TrackIndex))
+		if (TrackIndex != SelectedCaptionTrack)
 		{
-			if (TrackIndex != SelectedCaptionTrack)
+			UE_LOG(LogAvfMedia, Verbose, TEXT("Selecting caption track %i instead of %i (%i tracks)."), TrackIndex, SelectedCaptionTrack, CaptionTracks.Num());
+
+			// disable current track
+			if (SelectedCaptionTrack != INDEX_NONE)
 			{
-				if (SelectedCaptionTrack != INDEX_NONE)
-				{
-					const FTrack& CurrTrack = CaptionTracks[SelectedCaptionTrack];
-					
-					[PlayerItem removeOutput:(AVPlayerItemOutput*)CurrTrack.Output];
-					PlayerItem.tracks[CurrTrack.StreamIndex].enabled = NO;
-					
-					SelectedCaptionTrack = INDEX_NONE;
-				}
+				UE_LOG(LogAvfMedia, VeryVerbose, TEXT("Disabling caption track %i"), SelectedCaptionTrack);
 
-//				if ((TrackIndex != INDEX_NONE) && SUCCEEDED(PresentationDescriptor->SelectStream(CaptionTracks[TrackIndex].StreamIndex)))
-				{
-					SelectedCaptionTrack = TrackIndex;
-				}
+				const FTrack& Track = CaptionTracks[SelectedCaptionTrack];
 
-				if (SelectedCaptionTrack == TrackIndex)
-				{
-					const FTrack& SelectedTrack = CaptionTracks[SelectedCaptionTrack];
-					PlayerItem.tracks[SelectedTrack.StreamIndex].enabled = YES;
-					
-					InitializeOverlaySink();
-				}
-			}
+				[PlayerItem removeOutput : (AVPlayerItemOutput*)Track.Output];
+				PlayerItem.tracks[Track.StreamIndex].enabled = NO;
 
-			return true;
-		}
-		break;
+				SelectedCaptionTrack = INDEX_NONE;
+			}
 
-	case EMediaTrackType::Video:
-		if ((TrackIndex == INDEX_NONE) || VideoTracks.IsValidIndex(TrackIndex))
-		{
-			if (TrackIndex != SelectedVideoTrack)
+			// enable new track
+			if (TrackIndex != INDEX_NONE)
 			{
-				if (SelectedVideoTrack != INDEX_NONE)
+				if (!CaptionTracks.IsValidIndex(TrackIndex))
 				{
-					const FTrack& CurrTrack = VideoTracks[SelectedVideoTrack];
-					
-					[PlayerItem removeOutput:(AVPlayerItemOutput*)CurrTrack.Output];
-					PlayerItem.tracks[CurrTrack.StreamIndex].enabled = NO;
-					
-					SelectedVideoTrack = INDEX_NONE;
+					return false;
 				}
 
-//				if ((TrackIndex != INDEX_NONE) && SUCCEEDED(PresentationDescriptor->SelectStream(VideoTracks[TrackIndex].StreamIndex)))
-				{
-					SelectedVideoTrack = TrackIndex;
-				}
+				UE_LOG(LogAvfMedia, VeryVerbose, TEXT("Enabling caption track %i"), TrackIndex);
 
-				if (SelectedVideoTrack == TrackIndex)
-				{
-					const FTrack& SelectedTrack = VideoTracks[SelectedVideoTrack];
-					PlayerItem.tracks[SelectedTrack.StreamIndex].enabled = YES;
-					
-					InitializeVideoSink();
-				}
+				const FTrack& SelectedTrack = CaptionTracks[TrackIndex];
+				PlayerItem.tracks[SelectedTrack.StreamIndex].enabled = YES;
 			}
 
-			return true;
-		}
-		break;
+			SelectedCaptionTrack = TrackIndex;
 
-	default:
+			// update output
+			if (SelectedCaptionTrack != INDEX_NONE)
+			{
+				[PlayerItem addOutput:(AVPlayerItemOutput*)CaptionTracks[SelectedCaptionTrack].Output];
+			}
+		}
 		break;
-	}
 
-	return false;
-}
+	case EMediaTrackType::Video:
+		if (TrackIndex != SelectedVideoTrack)
+		{
+			UE_LOG(LogAvfMedia, Verbose, TEXT("Selecting video track %i instead of %i (%i tracks)"), TrackIndex, SelectedVideoTrack, VideoTracks.Num());
 
+			// disable current track
+			if (SelectedVideoTrack != INDEX_NONE)
+			{
+				UE_LOG(LogAvfMedia, VeryVerbose, TEXT("Disabling video track %i"), SelectedVideoTrack);
 
-/* FAvfMediaTracks implementation
- *****************************************************************************/
+				const FTrack& Track = VideoTracks[SelectedVideoTrack];
 
-void FAvfMediaTracks::InitializeAudioSink()
-{
-	if ((AudioSink == nullptr) || (SelectedAudioTrack == INDEX_NONE))
-	{
-		return;
-	}
+				[PlayerItem removeOutput : (AVPlayerItemOutput*)Track.Output];
+				PlayerItem.tracks[Track.StreamIndex].enabled = NO;
 
-#if PLATFORM_MAC
-	CMFormatDescriptionRef DescRef = (CMFormatDescriptionRef)[AudioTracks[SelectedAudioTrack].AssetTrack.formatDescriptions objectAtIndex:0];
-	AudioStreamBasicDescription const* ASBD = CMAudioFormatDescriptionGetStreamBasicDescription(DescRef);
-	
-	TargetDesc.mSampleRate = ASBD->mSampleRate;
-	TargetDesc.mFormatID = kAudioFormatLinearPCM;
-	TargetDesc.mFormatFlags = kAudioFormatFlagIsSignedInteger | kLinearPCMFormatFlagIsPacked;
-	TargetDesc.mBytesPerPacket = ASBD->mChannelsPerFrame * sizeof(int16);
-	TargetDesc.mFramesPerPacket = 1;
-	TargetDesc.mBytesPerFrame = ASBD->mChannelsPerFrame * sizeof(int16);
-	TargetDesc.mChannelsPerFrame = ASBD->mChannelsPerFrame;
-	TargetDesc.mBitsPerChannel = 16;
-	TargetDesc.mReserved = 0;
-#else
-	AVPlayerItemTrack* PlayerTrack = (AVPlayerItemTrack*)AudioTracks[SelectedAudioTrack].Output;
-	check(PlayerTrack);
-	PlayerTrack.enabled = true;
-#endif
+				SelectedVideoTrack = INDEX_NONE;
+			}
 
-	AudioSink->InitializeAudioSink(
-		GetAudioTrackChannels(SelectedAudioTrack),
-		GetAudioTrackSampleRate(SelectedAudioTrack)
-	);
-	
-	AudioSink->FlushAudioSink();
-}
+			// enable new track
+			if (TrackIndex != INDEX_NONE)
+			{
+				if (!VideoTracks.IsValidIndex(TrackIndex))
+				{
+					return false;
+				}
 
+				UE_LOG(LogAvfMedia, VeryVerbose, TEXT("Enabling video track %i"), TrackIndex);
 
-void FAvfMediaTracks::InitializeOverlaySink()
-{
-	if ((OverlaySink == nullptr) || (SelectedCaptionTrack == INDEX_NONE))
-	{
-		return;
-	}
-	
-	[PlayerItem addOutput:(AVPlayerItemOutput*)CaptionTracks[SelectedCaptionTrack].Output];
+				const FTrack& SelectedTrack = VideoTracks[TrackIndex];
+				PlayerItem.tracks[SelectedTrack.StreamIndex].enabled = YES;
+			}
 
-	OverlaySink->InitializeOverlaySink();
-}
+			SelectedVideoTrack = TrackIndex;
 
+			// update output
+			if (SelectedVideoTrack != INDEX_NONE)
+			{
+				[PlayerItem addOutput : (AVPlayerItemOutput*)VideoTracks[SelectedVideoTrack].Output];
 
-void FAvfMediaTracks::InitializeVideoSink()
-{
-	if ((VideoSink == nullptr) || (SelectedVideoTrack == INDEX_NONE))
-	{
-		return;
-	}
+				struct FSetOutputParams
+				{
+					AVPlayerItemVideoOutput* Output;
+					TWeakPtr<FAvfMediaVideoSampler, ESPMode::ThreadSafe> VideoSamplerPtr;
+					float FrameRate;
+				}
+				SetOutputParams = {
+					(AVPlayerItemVideoOutput*)VideoTracks[SelectedVideoTrack].Output,
+					VideoSampler,
+					1.0f / [VideoTracks[TrackIndex].AssetTrack nominalFrameRate]
+				};
+
+				ENQUEUE_UNIQUE_RENDER_COMMAND_ONEPARAMETER(AvfMediaVideoSamplerSetOutput,
+					FSetOutputParams, Params, SetOutputParams,
+					{
+						auto PinnedVideoSampler = Params.VideoSamplerPtr.Pin();
 
-	[PlayerItem addOutput:(AVPlayerItemOutput*)VideoTracks[SelectedVideoTrack].Output];
+						if (PinnedVideoSampler.IsValid())
+						{
+							PinnedVideoSampler->SetOutput(Params.Output, Params.FrameRate);
+						}
+					});
+			}
+		}
+		break;
 
-	const FIntPoint Dimensions = GetVideoTrackDimensions(SelectedVideoTrack);
+	default:
+		return false;
+	}
 
-	VideoSink->InitializeTextureSink(
-		Dimensions,
-		Dimensions,
-		EMediaTextureSinkFormat::CharBGRA,
-#if WITH_ENGINE
-		EMediaTextureSinkMode::Unbuffered
-#else
-		EMediaTextureSinkMode::Buffered
-#endif
-	);
-	
-	ENQUEUE_UNIQUE_RENDER_COMMAND_THREEPARAMETER(
-		InitialiseAvfVideoSampler,
-		FAvfMediaTracks*, This, this,
-		IMediaTextureSink*, VideoSink, VideoSink,
-		AVPlayerItemVideoOutput*, Output, (AVPlayerItemVideoOutput*)VideoTracks[SelectedVideoTrack].Output,
-    {
-		This->VideoSampler->SetTrack(VideoSink, Output);
-    });
+	return true;
 }
 
-FAvfVideoSampler::FAvfVideoSampler()
-: VideoSink(nullptr)
-, Output(nil)
-#if WITH_ENGINE && COREVIDEO_SUPPORTS_METAL
-, MetalTextureCache(nullptr)
-#endif
-{
-	
-}
 
-FAvfVideoSampler::~FAvfVideoSampler()
+bool FAvfMediaTracks::SetTrackFormat(EMediaTrackType TrackType, int32 TrackIndex, int32 FormatIndex)
 {
-	[Output release];
-#if WITH_ENGINE && COREVIDEO_SUPPORTS_METAL
-	if (MetalTextureCache)
+	if (FormatIndex != 0)
 	{
-		CFRelease(MetalTextureCache);
-		MetalTextureCache = nullptr;
+		return false;
 	}
-#endif
-}
 
-void FAvfVideoSampler::SetTrack(IMediaTextureSink* InVideoSink, AVPlayerItemVideoOutput* InOutput)
-{
 	FScopeLock Lock(&CriticalSection);
-	VideoSink = InVideoSink;
-	[InOutput retain];
-	[Output release];
-	Output = InOutput;
-}
-
-/* FTickableObjectRenderThread interface
- *****************************************************************************/
 
-TStatId FAvfVideoSampler::GetStatId() const
-{
-	RETURN_QUICK_DECLARE_CYCLE_STAT(FAvfVideoSampler, STATGROUP_Tickables);
-}
+	switch (TrackType)
+	{
+	case EMediaTrackType::Audio:
+		return AudioTracks.IsValidIndex(TrackIndex);
 
-bool FAvfVideoSampler::IsTickable() const
-{
-	return ((VideoSink != nullptr) && (Output != nil));
-}
+	case EMediaTrackType::Caption:
+		return CaptionTracks.IsValidIndex(TrackIndex);
 
-void FAvfVideoSampler::Tick(float /*DeltaTime*/)
-{
-	FScopeLock Lock(&CriticalSection);
-	
-	if ((VideoSink != nullptr) && (Output != nil))
-	{
-		CMTime OutputItemTime = [Output itemTimeForHostTime:CACurrentMediaTime()];
-		if ([Output hasNewPixelBufferForItemTime:OutputItemTime])
-		{
-			CVPixelBufferRef Frame =  [Output copyPixelBufferForItemTime:OutputItemTime itemTimeForDisplay:nullptr];
-			if(Frame)
-			{
-#if WITH_ENGINE
-#if COREVIDEO_SUPPORTS_METAL // On iOS/tvOS we use the Metal texture cache, we can do the same on macOS only when the min.version is 10.11+
-				if (IsMetalPlatform(GMaxRHIShaderPlatform))
-				{
-					if (!MetalTextureCache)
-					{
-						id<MTLDevice> Device = (id<MTLDevice>)GDynamicRHI->RHIGetNativeDevice();
-						check(Device);
-					
-						CVReturn Return = CVMetalTextureCacheCreate(kCFAllocatorDefault, nullptr, Device, nullptr, &MetalTextureCache);
-						check(Return == kCVReturnSuccess);
-					}
-					check(MetalTextureCache);
-                    
-                    if (CVPixelBufferIsPlanar(Frame))
-                    {
-                        check(IsMetalPlatform(GMaxRHIShaderPlatform));
-                        
-                        uint32 TexCreateFlags = TexCreate_Dynamic | TexCreate_NoTiling;
-                        
-                        int32 YWidth = CVPixelBufferGetWidthOfPlane(Frame, 0);
-                        int32 YHeight = CVPixelBufferGetHeightOfPlane(Frame, 0);
-                        
-                        CVMetalTextureRef YTextureRef = nullptr;
-                        CVReturn Result = CVMetalTextureCacheCreateTextureFromImage(kCFAllocatorDefault, MetalTextureCache, Frame, nullptr, MTLPixelFormatR8Unorm, YWidth, YHeight, 0, &YTextureRef);
-                        check(Result == kCVReturnSuccess);
-                        check(YTextureRef);
-                        
-                        int32 UVWidth = CVPixelBufferGetWidthOfPlane(Frame, 1);
-                        int32 UVHeight = CVPixelBufferGetHeightOfPlane(Frame, 1);
-                        
-                        CVMetalTextureRef UVTextureRef = nullptr;
-                        Result = CVMetalTextureCacheCreateTextureFromImage(kCFAllocatorDefault, MetalTextureCache, Frame, nullptr, MTLPixelFormatRG8Unorm, UVWidth, UVHeight, 1, &UVTextureRef);
-                        check(Result == kCVReturnSuccess);
-                        check(UVTextureRef);
-                        
-                        // Metal can upload directly from an IOSurface to a 2D texture, so we can just wrap it.
-                        FRHIResourceCreateInfo YCreateInfo;
-                        YCreateInfo.BulkData = new FAvfTexture2DResourceWrapper(YTextureRef);
-                        YCreateInfo.ResourceArray = nullptr;
-                        
-                        FRHIResourceCreateInfo UVCreateInfo;
-                        UVCreateInfo.BulkData = new FAvfTexture2DResourceWrapper(UVTextureRef);
-                        UVCreateInfo.ResourceArray = nullptr;
-                        
-                        TRefCountPtr<FRHITexture2D> YTex = RHICreateTexture2D(YWidth, YHeight, PF_G8, 1, 1, TexCreateFlags | TexCreate_ShaderResource, YCreateInfo);
-                        TRefCountPtr<FRHITexture2D> UVTex = RHICreateTexture2D(UVWidth, UVHeight, PF_R8G8, 1, 1, TexCreateFlags | TexCreate_ShaderResource, UVCreateInfo);
-                        
-                        FRHIResourceCreateInfo Info;
-                        TRefCountPtr<FRHITexture2D> ShaderResource = RHICreateTexture2D(YWidth, YHeight, PF_B8G8R8A8, 1, 1, TexCreateFlags | TexCreate_ShaderResource | TexCreate_RenderTargetable, Info);
-                        
-                        // render video frame into sink texture
-                        FRHICommandListImmediate& RHICmdList = FRHICommandListExecutor::GetImmediateCommandList();
-                        {
-                            // configure media shaders
-                            auto ShaderMap = GetGlobalShaderMap(ERHIFeatureLevel::SM5);
-                            TShaderMapRef<FMediaShadersVS> VertexShader(ShaderMap);
-                            TShaderMapRef<FYCbCrConvertPS> PixelShader(ShaderMap);
-                            
-                            SetRenderTarget(RHICmdList, ShaderResource, nullptr, ESimpleRenderTargetMode::EExistingColorAndDepth, FExclusiveDepthStencil::DepthNop_StencilNop);
-                            
-                            FGraphicsPipelineStateInitializer GraphicsPSOInit;
-                            RHICmdList.ApplyCachedRenderTargets(GraphicsPSOInit);
-                            
-                            GraphicsPSOInit.BlendState = TStaticBlendStateWriteMask<CW_RGBA, CW_NONE, CW_NONE, CW_NONE, CW_NONE, CW_NONE, CW_NONE, CW_NONE>::GetRHI();
-                            GraphicsPSOInit.RasterizerState = TStaticRasterizerState<>::GetRHI();
-                            GraphicsPSOInit.DepthStencilState = TStaticDepthStencilState<false, CF_Always>::GetRHI();
-                            
-                            GraphicsPSOInit.BoundShaderState.VertexDeclarationRHI = GMediaVertexDeclaration.VertexDeclarationRHI;
-                            GraphicsPSOInit.BoundShaderState.VertexShaderRHI = GETSAFERHISHADER_VERTEX(*VertexShader);
-                            GraphicsPSOInit.BoundShaderState.PixelShaderRHI = GETSAFERHISHADER_PIXEL(*PixelShader);
-                            GraphicsPSOInit.PrimitiveType = PT_TriangleStrip;
-                            
-                            SetGraphicsPipelineState(RHICmdList, GraphicsPSOInit);
-                            
-                            PixelShader->SetParameters(RHICmdList, YTex, UVTex, MediaShaders::YuvToSrgbPs4, true);
-                            
-                            // draw full-size quad
-                            FMediaElementVertex Vertices[4];
-                            Vertices[0].Position.Set(-1.0f,  1.0f, 1.0f, 1.0f); // Top Left
-                            Vertices[1].Position.Set( 1.0f,  1.0f, 1.0f, 1.0f); // Top Right
-                            Vertices[2].Position.Set(-1.0f, -1.0f, 1.0f, 1.0f); // Bottom Left
-                            Vertices[3].Position.Set( 1.0f, -1.0f, 1.0f, 1.0f); // Bottom Right
-                            
-                            const float ULeft   = 0.0f;
-                            const float URight  = 1.0f;
-                            const float VTop    = 0.0f;
-                            const float VBottom = 1.0f;
-                            
-                            Vertices[0].TextureCoordinate.Set(ULeft, VTop);
-                            Vertices[1].TextureCoordinate.Set(URight, VTop);
-                            Vertices[2].TextureCoordinate.Set(ULeft, VBottom);
-                            Vertices[3].TextureCoordinate.Set(URight, VBottom);
-                            
-                            RHICmdList.SetViewport(0, 0, 0.0f, YWidth, YHeight, 1.0f);
-                            
-                            
-                            DrawPrimitiveUP(RHICmdList, PT_TriangleStrip, 2, Vertices, sizeof(Vertices[0]));
-                            
-                            RHICmdList.CopyToResolveTarget(ShaderResource, ShaderResource, true, FResolveParams());
-                        }
-                        
-                        
-                        VideoSink->UpdateTextureSinkResource(ShaderResource, ShaderResource);
-                        CFRelease(YTextureRef);
-                        CFRelease(UVTextureRef);
-                    }
-                    else
-                    {
-                        int32 Width = CVPixelBufferGetWidth(Frame);
-                        int32 Height = CVPixelBufferGetHeight(Frame);
-                        
-                        CVMetalTextureRef TextureRef = nullptr;
-                        CVReturn Result = CVMetalTextureCacheCreateTextureFromImage(kCFAllocatorDefault, MetalTextureCache, Frame, nullptr, MTLPixelFormatBGRA8Unorm_sRGB, Width, Height, 0, &TextureRef);
-                        check(Result == kCVReturnSuccess);
-                        check(TextureRef);
-                        
-                        FRHIResourceCreateInfo CreateInfo;
-                        CreateInfo.BulkData = new FAvfTexture2DResourceWrapper(TextureRef);
-                        CreateInfo.ResourceArray = nullptr;
-                        
-                        uint32 TexCreateFlags = TexCreate_SRGB;
-                        TexCreateFlags |= TexCreate_Dynamic | TexCreate_NoTiling;
-                        
-                        TRefCountPtr<FRHITexture2D> RenderTarget;
-                        TRefCountPtr<FRHITexture2D> ShaderResource;
-                        
-                        ShaderResource = RHICreateTexture2D(Width, Height, PF_B8G8R8A8, 1, 1, TexCreateFlags | TexCreate_ShaderResource, CreateInfo);
-                        
-                        VideoSink->UpdateTextureSinkResource(ShaderResource, ShaderResource);
-                        CFRelease(TextureRef);
-                    }
-				}
-				else // Ran out of time to implement efficient OpenGLES texture upload - its running out of memory.
-				/*{
-					if (!OpenGLTextureCache)
-					{
-						EAGLContext* Context = [EAGLContext currentContext];
-						check(Context);
-						
-						CVReturn Return = CVOpenGLESTextureCacheCreate(kCFAllocatorDefault, nullptr, Context, nullptr, &OpenGLTextureCache);
-						check(Return == kCVReturnSuccess);
-					}
-					check(OpenGLTextureCache);
-				
-					int32 Width = CVPixelBufferGetWidth(Frame);
-					int32 Height = CVPixelBufferGetHeight(Frame);
-					
-					CVOpenGLESTextureRef TextureRef = nullptr;
-					CVReturn Result = CVOpenGLESTextureCacheCreateTextureFromImage(kCFAllocatorDefault, OpenGLTextureCache, Frame, nullptr, GL_TEXTURE_2D, GL_RGBA, Width, Height, GL_RGBA, GL_UNSIGNED_BYTE, 0, &TextureRef);
-					check(Result == kCVReturnSuccess);
-					check(TextureRef);
-				
-					FRHIResourceCreateInfo CreateInfo;
-					CreateInfo.BulkData = new FAvfTexture2DResourceWrapper(TextureRef);
-					CreateInfo.ResourceArray = nullptr;
-					
-					uint32 TexCreateFlags = 0;
-					TexCreateFlags |= TexCreate_Dynamic | TexCreate_NoTiling;
-					
-					TRefCountPtr<FRHITexture2D> RenderTarget;
-					TRefCountPtr<FRHITexture2D> ShaderResource;
-					
-					RHICreateTargetableShaderResource2D(Width,
-														Height,
-														PF_R8G8B8A8,
-														1,
-														TexCreateFlags,
-														TexCreate_RenderTargetable,
-														false,
-														CreateInfo,
-														RenderTarget,
-														ShaderResource
-														);
-														
-					VideoSink->UpdateTextureSinkResource(RenderTarget, ShaderResource);
-					CFRelease(TextureRef);
-				}*/
-#endif	// On Mac we have to use IOSurfaceRef for backward compatibility - unless we update MIN_REQUIRED_VERSION to 10.11 we link against an older version of CoreVideo that doesn't support Metal.
-				{
-					FRHIResourceCreateInfo CreateInfo;
-					if (IsMetalPlatform(GMaxRHIShaderPlatform))
-					{
-						// Metal can upload directly from an IOSurface to a 2D texture, so we can just wrap it.
-						CreateInfo.BulkData = new FAvfTexture2DResourceWrapper(Frame);
-					}
-					else
-					{
-						// OpenGL on Mac uploads as a TEXTURE_RECTANGLE for which we have no code, so upload via system memory.
-						CreateInfo.BulkData = new FAvfTexture2DResourceMem(Frame);
-					}
-					CreateInfo.ResourceArray = nullptr;
-					
-					int32 Width = CVPixelBufferGetWidth(Frame);
-					int32 Height = CVPixelBufferGetHeight(Frame);
-					
-					uint32 TexCreateFlags = TexCreate_SRGB;
-					TexCreateFlags |= TexCreate_Dynamic | TexCreate_NoTiling;
-					
-					TRefCountPtr<FRHITexture2D> RenderTarget;
-					TRefCountPtr<FRHITexture2D> ShaderResource;
-					ShaderResource = RHICreateTexture2D(Width, Height, PF_B8G8R8A8, 1, 1, TexCreateFlags | TexCreate_ShaderResource, CreateInfo);
-					
-					VideoSink->UpdateTextureSinkResource(ShaderResource, ShaderResource);
- 				}
-#else
-				int32 Pitch = CVPixelBufferGetBytesPerRow(Frame);
+	case EMediaTrackType::Video:
+		return VideoTracks.IsValidIndex(TrackIndex);
 
-				CVPixelBufferLockBaseAddress(Frame, kCVPixelBufferLock_ReadOnly);
-				
-				uint8* VideoData = (uint8*)CVPixelBufferGetBaseAddress(Frame);
-				VideoSink->UpdateTextureSinkBuffer(VideoData, Pitch);
-				
-				CVPixelBufferUnlockBaseAddress(Frame, kCVPixelBufferLock_ReadOnly);
-#endif
-				FTimespan DiplayTime = FTimespan::FromSeconds(CMTimeGetSeconds(OutputItemTime));
-				VideoSink->DisplayTextureSinkBuffer(DiplayTime);
-				
-				CVPixelBufferRelease(Frame);
-			}
-		}
+	default:
+		return false;
 	}
 }
+
diff --git a/Engine/Plugins/Media/AvfMedia/Source/AvfMedia/Private/Player/AvfMediaTracks.h b/Engine/Plugins/Media/AvfMedia/Source/AvfMedia/Private/Player/AvfMediaTracks.h
--- a/Engine/Plugins/Media/AvfMedia/Source/AvfMedia/Private/Player/AvfMediaTracks.h
+++ b/Engine/Plugins/Media/AvfMedia/Source/AvfMedia/Private/Player/AvfMediaTracks.h
@@ -2,49 +2,27 @@
 
 #pragma once
 
-#include "../AvfMediaPrivate.h"
-
-#include "IMediaOutput.h"
+#include "CoreTypes.h"
+#include "Containers/UnrealString.h"
+#include "HAL/CriticalSection.h"
+#include "IMediaSamples.h"
 #include "IMediaTracks.h"
+#include "Internationalization/Text.h"
+#include "Templates/SharedPointer.h"
 
-#include "TickableObjectRenderThread.h"
+#import <AVFoundation/AVFoundation.h>
 
-class FAvfVideoSampler : public FTickableObjectRenderThread
-{
-public:
-	FAvfVideoSampler();
-	virtual ~FAvfVideoSampler();
-	
-	void SetTrack(IMediaTextureSink* VideoSink, AVPlayerItemVideoOutput* Output);
+class FAvfMediaAudioSamplePool;
+class FAvfMediaVideoSamplePool;
+class FAvfMediaVideoSampler;
+class FMediaSamples;
 
-public:
-	/* FTickableObjectRenderThread interface */
-	virtual TStatId GetStatId() const override;
-	
-	virtual bool IsTickable() const override;
-	
-	virtual void Tick(float /*DeltaTime*/) override;
-	
-private:
-	/** Mutex to ensure thread-safe access */
-	FCriticalSection CriticalSection;
-	
-	/** The video sink */
-	IMediaTextureSink* VideoSink;
-	
-	/** The track's video output handle */
-	AVPlayerItemVideoOutput* Output;
-	
-#if WITH_ENGINE && COREVIDEO_SUPPORTS_METAL
-private:
-	/** The Metal texture cache for unbuffered texture uploads. */
-	CVMetalTextureCacheRef MetalTextureCache;
-#endif
-};
 
+/**
+ * 
+ */
 class FAvfMediaTracks
-	: public IMediaOutput
-	, public IMediaTracks
+	: public IMediaTracks
 {
 	enum ESyncStatus
 	{
@@ -69,15 +47,19 @@ class FAvfMediaTracks
 public:
 
 	/** Default constructor. */
-	FAvfMediaTracks();
+	FAvfMediaTracks(FMediaSamples& InSamples);
 
 	/** Virtual destructor. */
 	virtual ~FAvfMediaTracks();
 
 public:
 
-	/** Discard any pending data on the media sinks. */
-	void Flush();
+	/**
+	 * Append track statistics information to the given string.
+	 *
+	 * @param OutStats The string to append the statistics to.
+	 */
+	void AppendStats(FString &OutStats) const;
 
 	/**
 	 * Initialize the track collection.
@@ -87,95 +69,65 @@ public:
 	 */
 	void Initialize(AVPlayerItem* InPlayerItem, FString& OutInfo);
 
-	/** Reset the stream collection. */
-	void Reset();
-
 	/**
-	 * Notify tracks that playback rate was changed.
+	 * Process audio frames.
 	 *
-	 * @param Rate The new playback rate.
+	 * @see ProcessCaptions, ProcessVideo
 	 */
-	void SetRate(float Rate);
-	
+	void ProcessAudio();
+
 	/**
-	 * Notify tracks that playback is seeking.
+	 * Process caption frames.
 	 *
-	 * @param Time The time to seek to.
+	 * Called by the caption track delegate to provide the attributed strings
+	 * for each timecode to the caption sink.
+	 *
+	 * @param Output The caption track output that generated the strings.
+	 * @param Strings The attributed caption strings.
+	 * @param ItemTime The display time for the strings.
+	 * @see ProcessAudio, ProcessVideo
 	 */
-	void Seek(const FTimespan& Time);
-	
+	void ProcessCaptions(AVPlayerItemLegibleOutput* Output, NSArray<NSAttributedString*>* Strings, NSArray* NativeSamples, CMTime ItemTime);
+
 	/**
-	 * Append track statistics information to the given string.
+	 * Process video frames.
 	 *
-	 * @param OutStats The string to append the statistics to.
+	 * @see ProcessAudio, ProcessCaptions
 	 */
-	void AppendStats(FString &OutStats) const;
+	void ProcessVideo();
+
+	/** Reset the stream collection. */
+	void Reset();
 
 	/**
-	 * Update the playback of tracks - must only be called when playing.
+	 * Notify tracks that playback is seeking.
 	 *
-	 * @param DeltaTime Time between ticks.
+	 * @param Time The time to seek to.
 	 */
-	bool Tick(float DeltaTime);
-	
+	void Seek(const FTimespan& Time);
+
 	/**
-	 * Called by the caption track delegate to provide the attributed strings for each timecode to the caption sink.
+	 * Notify tracks that playback rate was changed.
 	 *
-	 * @param Output The caption track output that generated the strings.
-	 * @param Strings The attributed caption strings.
-	 * @param ItemTime The display time for the strings.
+	 * @param Rate The new playback rate.
 	 */
-	void HandleSubtitleCaptionStrings(AVPlayerItemLegibleOutput* Output, NSArray<NSAttributedString*>* Strings, NSArray* NativeSamples, CMTime ItemTime);
-
-public:
-
-	//~ IMediaOutput interface
-
-	virtual void SetAudioSink(IMediaAudioSink* Sink) override;
-	virtual void SetMetadataSink(IMediaBinarySink* Sink) override;
-	virtual void SetOverlaySink(IMediaOverlaySink* Sink) override;
-	virtual void SetVideoSink(IMediaTextureSink* Sink) override;
-
+	void SetRate(float Rate);
+	
 public:
 
 	//~ IMediaTracks interface
 
-	virtual uint32 GetAudioTrackChannels(int32 TrackIndex) const override;
-	virtual uint32 GetAudioTrackSampleRate(int32 TrackIndex) const override;
+	virtual bool GetAudioTrackFormat(int32 TrackIndex, int32 FormatIndex, FMediaAudioTrackFormat& OutFormat) const override;
 	virtual int32 GetNumTracks(EMediaTrackType TrackType) const override;
+	virtual int32 GetNumTrackFormats(EMediaTrackType TrackType, int32 TrackIndex) const override;
 	virtual int32 GetSelectedTrack(EMediaTrackType TrackType) const override;
 	virtual FText GetTrackDisplayName(EMediaTrackType TrackType, int32 TrackIndex) const override;
+	virtual int32 GetTrackFormat(EMediaTrackType TrackType, int32 TrackIndex) const override;
 	virtual FString GetTrackLanguage(EMediaTrackType TrackType, int32 TrackIndex) const override;
 	virtual FString GetTrackName(EMediaTrackType TrackType, int32 TrackIndex) const override;
-	virtual uint32 GetVideoTrackBitRate(int32 TrackIndex) const override;
-	virtual FIntPoint GetVideoTrackDimensions(int32 TrackIndex) const override;
-	virtual float GetVideoTrackFrameRate(int32 TrackIndex) const override;
+	virtual bool GetVideoTrackFormat(int32 TrackIndex, int32 FormatIndex, FMediaVideoTrackFormat& OutFormat) const override;
 	virtual bool SelectTrack(EMediaTrackType TrackType, int32 TrackIndex) override;
-
-protected:
-
-	/** Initialize the current audio sink. */
-	void InitializeAudioSink();
-
-	/** Initialize the current text overlay sink. */
-	void InitializeOverlaySink();
-
-	/** Initialize the current video sink. */
-	void InitializeVideoSink();
-	
-	/** Update the current video sink. */
-	void UpdateVideoSink();
-
-private:
-
-	/** The currently used audio sink. */
-	IMediaAudioSink* AudioSink;
-
-	/** The currently used text overlay sink. */
-	IMediaOverlaySink* OverlaySink;
-
-	/** The currently used video sink. */
-	IMediaTextureSink* VideoSink;
+	virtual bool SetTrackFormat(EMediaTrackType TrackType, int32 TrackIndex, int32 FormatIndex) override;
 
 private:
 
@@ -190,14 +142,26 @@ private:
 
 private:
 
+	/** Whether the audio is currently paused. */
+	bool AudioPaused;
+
+	/** Audio sample object pool. */
+	FAvfMediaAudioSamplePool* AudioSamplePool;
+
 	/** Synchronizes write access to track arrays, selections & sinks. */
 	mutable FCriticalSection CriticalSection;
 
+	/** The last audio sample provided to the sink. */
+	CMTime LastAudioSampleTime;
+
 	/** The player item containing the track information. */
 	AVPlayerItem* PlayerItem;
+
+	/** The media sample queue. */
+	FMediaSamples& Samples;
 	
-	/** Object to sample video frames */
-	FAvfVideoSampler* VideoSampler;
+	/** Seek to this time. */
+	double SeekTime;
 
 	/** Index of the selected audio track. */
 	int32 SelectedAudioTrack;
@@ -208,20 +172,12 @@ private:
 	/** Index of the selected video track. */
 	int32 SelectedVideoTrack;
 
-private:
-	
-	/** Target desciption for audio output required by Media framework audio sinks. */
+	/** Target description for audio output required by Media framework audio sinks. */
 	AudioStreamBasicDescription TargetDesc;
 	
-	/** Whether the audio is currently paused. */
-	bool bAudioPaused;
-	
-	/** The last audio sample provided to the sink. */
-	CMTime LastAudioSample;
-	
-	/** Seek to this time. */
-	double SeekTime;
-	
+	/** Object to sample video frames */
+	TSharedPtr<FAvfMediaVideoSampler, ESPMode::ThreadSafe> VideoSampler;
+		
 	/** Has been played with fast/slow rate? */
-	bool bZoomed;
+	bool Zoomed;
 };
diff --git a/Engine/Plugins/Media/AvfMedia/Source/AvfMedia/Private/Player/AvfMediaVideoSampler.cpp b/Engine/Plugins/Media/AvfMedia/Source/AvfMedia/Private/Player/AvfMediaVideoSampler.cpp
new file mode 100644
--- /dev/null
+++ b/Engine/Plugins/Media/AvfMedia/Source/AvfMedia/Private/Player/AvfMediaVideoSampler.cpp
@@ -0,0 +1,435 @@
+// Copyright 1998-2017 Epic Games, Inc. All Rights Reserved.
+
+#include "AvfMediaVideoSampler.h"
+#include "AvfMediaPrivate.h"
+
+#include "Containers/ResourceArray.h"
+#include "MediaSamples.h"
+#include "Misc/ScopeLock.h"
+
+#if WITH_ENGINE
+#include "RenderingThread.h"
+#include "RHI.h"
+#include "MediaShaders.h"
+#include "StaticBoundShaderState.h"
+#include "Misc/ScopeLock.h"
+#include "PipelineStateCache.h"
+#include "RHIStaticStates.h"
+#include "Class.h"
+#endif
+
+#include "AvfMediaTextureSample.h"
+
+
+/**
+ * Passes a CV*TextureRef or CVPixelBufferRef through to the RHI to wrap in an RHI texture without traversing system memory.
+ */
+class FAvfTexture2DResourceWrapper
+	: public FResourceBulkDataInterface
+{
+public:
+
+	FAvfTexture2DResourceWrapper(CFTypeRef InImageBuffer)
+		: ImageBuffer(InImageBuffer)
+	{
+		check(ImageBuffer);
+		CFRetain(ImageBuffer);
+	}
+
+	virtual ~FAvfTexture2DResourceWrapper()
+	{
+		CFRelease(ImageBuffer);
+		ImageBuffer = nullptr;
+	}
+
+public:
+
+	//~ FResourceBulkDataInterface interface
+
+	virtual void Discard() override
+	{
+		delete this;
+	}
+
+	virtual const void* GetResourceBulkData() const override
+	{
+		return ImageBuffer;
+	}
+	
+	virtual uint32 GetResourceBulkDataSize() const override
+	{
+		return ImageBuffer ? ~0u : 0;
+	}
+
+	virtual EBulkDataType GetResourceType() const override
+	{
+		return EBulkDataType::MediaTexture;
+	}
+			
+	CFTypeRef ImageBuffer;
+};
+
+
+/**
+ * Allows for direct GPU mem allocation for texture resource from a CVImageBufferRef's system memory backing store.
+ */
+class FAvfTexture2DResourceMem
+	: public FResourceBulkDataInterface
+{
+public:
+	FAvfTexture2DResourceMem(CVImageBufferRef InImageBuffer)
+	: ImageBuffer(InImageBuffer)
+	{
+		check(ImageBuffer);
+		CFRetain(ImageBuffer);
+	}
+	
+	/**
+	 * @return ptr to the resource memory which has been preallocated
+	 */
+	virtual const void* GetResourceBulkData() const override
+	{
+		CVPixelBufferLockBaseAddress(ImageBuffer, kCVPixelBufferLock_ReadOnly);
+		return CVPixelBufferGetBaseAddress(ImageBuffer);
+	}
+	
+	/**
+	 * @return size of resource memory
+	 */
+	virtual uint32 GetResourceBulkDataSize() const override
+	{
+		int32 Pitch = CVPixelBufferGetBytesPerRow(ImageBuffer);
+		int32 Height = CVPixelBufferGetHeight(ImageBuffer);
+		uint32 Size = (Pitch * Height);
+
+		return Size;
+	}
+	
+	/**
+	 * Free memory after it has been used to initialize RHI resource
+	 */
+	virtual void Discard() override
+	{
+		CVPixelBufferUnlockBaseAddress(ImageBuffer, kCVPixelBufferLock_ReadOnly);
+		delete this;
+	}
+	
+	virtual ~FAvfTexture2DResourceMem()
+	{
+		CFRelease(ImageBuffer);
+		ImageBuffer = nullptr;
+	}
+	
+	CVImageBufferRef ImageBuffer;
+};
+
+
+/* FAvfMediaVideoSampler structors
+ *****************************************************************************/
+
+FAvfMediaVideoSampler::FAvfMediaVideoSampler(FMediaSamples& InSamples)
+	: Output(nil)
+	, Samples(InSamples)
+	, VideoSamplePool(new FAvfMediaTextureSamplePool)
+	, FrameRate(0.0f)
+#if WITH_ENGINE && COREVIDEO_SUPPORTS_METAL
+	, MetalTextureCache(nullptr)
+#endif
+{ }
+
+
+FAvfMediaVideoSampler::~FAvfMediaVideoSampler()
+{
+	[Output release];
+
+	delete VideoSamplePool;
+	VideoSamplePool = nullptr;
+
+#if WITH_ENGINE && COREVIDEO_SUPPORTS_METAL
+	if (MetalTextureCache)
+	{
+		CFRelease(MetalTextureCache);
+		MetalTextureCache = nullptr;
+	}
+#endif
+}
+
+
+/* FAvfMediaVideoSampler interface
+ *****************************************************************************/
+
+void FAvfMediaVideoSampler::SetOutput(AVPlayerItemVideoOutput* InOutput, float InFrameRate)
+{
+	check(IsInRenderingThread());
+
+	FScopeLock Lock(&CriticalSection);
+
+	[InOutput retain];
+	[Output release];
+	Output = InOutput;
+
+	FrameRate = InFrameRate;
+}
+
+
+void FAvfMediaVideoSampler::Tick()
+{
+	check(IsInRenderingThread());
+
+	FScopeLock Lock(&CriticalSection);
+
+	if (Output == nil)
+	{
+		return;
+	}
+
+	CMTime OutputItemTime = [Output itemTimeForHostTime:CACurrentMediaTime()];
+
+	if (![Output hasNewPixelBufferForItemTime : OutputItemTime])
+	{
+		return;
+	}
+
+	CVPixelBufferRef Frame = [Output copyPixelBufferForItemTime:OutputItemTime itemTimeForDisplay:nullptr];
+
+	if (!Frame)
+	{
+		return;
+	}
+
+	const FTimespan SampleDuration = FTimespan::FromSeconds(FrameRate);
+	const FTimespan SampleTime = FTimespan::FromSeconds(CMTimeGetSeconds(OutputItemTime));
+
+	const int32 FrameHeight = CVPixelBufferGetHeight(Frame);
+	const int32 FrameWidth = CVPixelBufferGetWidth(Frame);
+	const int32 FrameStride = CVPixelBufferGetBytesPerRow(Frame);
+
+	const FIntPoint Dim = FIntPoint(FrameStride / 4, FrameHeight);
+	const FIntPoint OutputDim = FIntPoint(FrameWidth, FrameHeight);
+
+	auto VideoSample = VideoSamplePool->AcquireShared();
+
+#if WITH_ENGINE
+	TRefCountPtr<FRHITexture2D> ShaderResource;
+#if COREVIDEO_SUPPORTS_METAL
+	// On iOS/tvOS we use the Metal texture cache
+	if (IsMetalPlatform(GMaxRHIShaderPlatform))
+	{
+		if (!MetalTextureCache)
+		{
+			id<MTLDevice> Device = (id<MTLDevice>)GDynamicRHI->RHIGetNativeDevice();
+			check(Device);
+			
+			CVReturn Return = CVMetalTextureCacheCreate(kCFAllocatorDefault, nullptr, Device, nullptr, &MetalTextureCache);
+			check(Return == kCVReturnSuccess);
+		}
+		check(MetalTextureCache);
+		
+		if (CVPixelBufferIsPlanar(Frame))
+		{
+			check(IsMetalPlatform(GMaxRHIShaderPlatform));
+			
+			uint32 TexCreateFlags = TexCreate_Dynamic | TexCreate_NoTiling;
+			
+			int32 YWidth = CVPixelBufferGetWidthOfPlane(Frame, 0);
+			int32 YHeight = CVPixelBufferGetHeightOfPlane(Frame, 0);
+			
+			CVMetalTextureRef YTextureRef = nullptr;
+			CVReturn Result = CVMetalTextureCacheCreateTextureFromImage(kCFAllocatorDefault, MetalTextureCache, Frame, nullptr, MTLPixelFormatR8Unorm, YWidth, YHeight, 0, &YTextureRef);
+			check(Result == kCVReturnSuccess);
+			check(YTextureRef);
+			
+			int32 UVWidth = CVPixelBufferGetWidthOfPlane(Frame, 1);
+			int32 UVHeight = CVPixelBufferGetHeightOfPlane(Frame, 1);
+			
+			CVMetalTextureRef UVTextureRef = nullptr;
+			Result = CVMetalTextureCacheCreateTextureFromImage(kCFAllocatorDefault, MetalTextureCache, Frame, nullptr, MTLPixelFormatRG8Unorm, UVWidth, UVHeight, 1, &UVTextureRef);
+			check(Result == kCVReturnSuccess);
+			check(UVTextureRef);
+			
+			// Metal can upload directly from an IOSurface to a 2D texture, so we can just wrap it.
+			FRHIResourceCreateInfo YCreateInfo;
+			YCreateInfo.BulkData = new FAvfTexture2DResourceWrapper(YTextureRef);
+			YCreateInfo.ResourceArray = nullptr;
+			
+			FRHIResourceCreateInfo UVCreateInfo;
+			UVCreateInfo.BulkData = new FAvfTexture2DResourceWrapper(UVTextureRef);
+			UVCreateInfo.ResourceArray = nullptr;
+			
+			TRefCountPtr<FRHITexture2D> YTex = RHICreateTexture2D(YWidth, YHeight, PF_G8, 1, 1, TexCreateFlags | TexCreate_ShaderResource, YCreateInfo);
+			TRefCountPtr<FRHITexture2D> UVTex = RHICreateTexture2D(UVWidth, UVHeight, PF_R8G8, 1, 1, TexCreateFlags | TexCreate_ShaderResource, UVCreateInfo);
+			
+			FRHIResourceCreateInfo Info;
+			ShaderResource = RHICreateTexture2D(YWidth, YHeight, PF_B8G8R8A8, 1, 1, TexCreateFlags | TexCreate_ShaderResource | TexCreate_RenderTargetable, Info);
+			
+			// render video frame into sink texture
+			FRHICommandListImmediate& RHICmdList = FRHICommandListExecutor::GetImmediateCommandList();
+			{
+				// configure media shaders
+				auto ShaderMap = GetGlobalShaderMap(GMaxRHIFeatureLevel);
+				TShaderMapRef<FMediaShadersVS> VertexShader(ShaderMap);
+				TShaderMapRef<FYCbCrConvertPS> PixelShader(ShaderMap);
+				
+				SetRenderTarget(RHICmdList, ShaderResource, nullptr, ESimpleRenderTargetMode::EExistingColorAndDepth, FExclusiveDepthStencil::DepthNop_StencilNop);
+				
+				FGraphicsPipelineStateInitializer GraphicsPSOInit;
+				RHICmdList.ApplyCachedRenderTargets(GraphicsPSOInit);
+				
+				GraphicsPSOInit.BlendState = TStaticBlendStateWriteMask<CW_RGBA, CW_NONE, CW_NONE, CW_NONE, CW_NONE, CW_NONE, CW_NONE, CW_NONE>::GetRHI();
+				GraphicsPSOInit.RasterizerState = TStaticRasterizerState<>::GetRHI();
+				GraphicsPSOInit.DepthStencilState = TStaticDepthStencilState<false, CF_Always>::GetRHI();
+				
+				GraphicsPSOInit.BoundShaderState.VertexDeclarationRHI = GMediaVertexDeclaration.VertexDeclarationRHI;
+				GraphicsPSOInit.BoundShaderState.VertexShaderRHI = GETSAFERHISHADER_VERTEX(*VertexShader);
+				GraphicsPSOInit.BoundShaderState.PixelShaderRHI = GETSAFERHISHADER_PIXEL(*PixelShader);
+				GraphicsPSOInit.PrimitiveType = PT_TriangleStrip;
+				
+				SetGraphicsPipelineState(RHICmdList, GraphicsPSOInit);
+				
+				PixelShader->SetParameters(RHICmdList, YTex, UVTex, MediaShaders::YuvToSrgbPs4, true);
+				
+				// draw full-size quad
+				FMediaElementVertex Vertices[4];
+				Vertices[0].Position.Set(-1.0f,  1.0f, 1.0f, 1.0f); // Top Left
+				Vertices[1].Position.Set( 1.0f,  1.0f, 1.0f, 1.0f); // Top Right
+				Vertices[2].Position.Set(-1.0f, -1.0f, 1.0f, 1.0f); // Bottom Left
+				Vertices[3].Position.Set( 1.0f, -1.0f, 1.0f, 1.0f); // Bottom Right
+				
+				const float ULeft   = 0.0f;
+				const float URight  = 1.0f;
+				const float VTop    = 0.0f;
+				const float VBottom = 1.0f;
+				
+				Vertices[0].TextureCoordinate.Set(ULeft, VTop);
+				Vertices[1].TextureCoordinate.Set(URight, VTop);
+				Vertices[2].TextureCoordinate.Set(ULeft, VBottom);
+				Vertices[3].TextureCoordinate.Set(URight, VBottom);
+				
+				RHICmdList.SetViewport(0, 0, 0.0f, YWidth, YHeight, 1.0f);
+
+				DrawPrimitiveUP(RHICmdList, PT_TriangleStrip, 2, Vertices, sizeof(Vertices[0]));
+				
+				RHICmdList.CopyToResolveTarget(ShaderResource, ShaderResource, true, FResolveParams());
+			}
+			
+			CFRelease(YTextureRef);
+			CFRelease(UVTextureRef);
+		}
+		else
+		{
+			int32 Width = CVPixelBufferGetWidth(Frame);
+			int32 Height = CVPixelBufferGetHeight(Frame);
+			
+			CVMetalTextureRef TextureRef = nullptr;
+			CVReturn Result = CVMetalTextureCacheCreateTextureFromImage(kCFAllocatorDefault, MetalTextureCache, Frame, nullptr, MTLPixelFormatBGRA8Unorm_sRGB, Width, Height, 0, &TextureRef);
+			check(Result == kCVReturnSuccess);
+			check(TextureRef);
+			
+			FRHIResourceCreateInfo CreateInfo;
+			CreateInfo.BulkData = new FAvfTexture2DResourceWrapper(TextureRef);
+			CreateInfo.ResourceArray = nullptr;
+			
+			uint32 TexCreateFlags = TexCreate_SRGB;
+			TexCreateFlags |= TexCreate_Dynamic | TexCreate_NoTiling;
+			
+			ShaderResource = RHICreateTexture2D(Width, Height, PF_B8G8R8A8, 1, 1, TexCreateFlags | TexCreate_ShaderResource, CreateInfo);
+			
+			CFRelease(TextureRef);
+		}
+	}
+	else // Ran out of time to implement efficient OpenGLES texture upload - its running out of memory.
+	/*{
+	 if (!OpenGLTextureCache)
+	 {
+	 EAGLContext* Context = [EAGLContext currentContext];
+	 check(Context);
+	 
+	 CVReturn Return = CVOpenGLESTextureCacheCreate(kCFAllocatorDefault, nullptr, Context, nullptr, &OpenGLTextureCache);
+	 check(Return == kCVReturnSuccess);
+	 }
+	 check(OpenGLTextureCache);
+	 
+	 int32 Width = CVPixelBufferGetWidth(Frame);
+	 int32 Height = CVPixelBufferGetHeight(Frame);
+	 
+	 CVOpenGLESTextureRef TextureRef = nullptr;
+	 CVReturn Result = CVOpenGLESTextureCacheCreateTextureFromImage(kCFAllocatorDefault, OpenGLTextureCache, Frame, nullptr, GL_TEXTURE_2D, GL_RGBA, Width, Height, GL_RGBA, GL_UNSIGNED_BYTE, 0, &TextureRef);
+	 check(Result == kCVReturnSuccess);
+	 check(TextureRef);
+	 
+	 FRHIResourceCreateInfo CreateInfo;
+	 CreateInfo.BulkData = new FAvfTexture2DResourceWrapper(TextureRef);
+	 CreateInfo.ResourceArray = nullptr;
+	 
+	 uint32 TexCreateFlags = 0;
+	 TexCreateFlags |= TexCreate_Dynamic | TexCreate_NoTiling;
+	 
+	 TRefCountPtr<FRHITexture2D> RenderTarget;
+	 TRefCountPtr<FRHITexture2D> ShaderResource;
+	 
+	 RHICreateTargetableShaderResource2D(Width,
+	 Height,
+	 PF_R8G8B8A8,
+	 1,
+	 TexCreateFlags,
+	 TexCreate_RenderTargetable,
+	 false,
+	 CreateInfo,
+	 RenderTarget,
+	 ShaderResource
+	 );
+	 
+	 VideoSink->UpdateTextureSinkResource(RenderTarget, ShaderResource);
+	 CFRelease(TextureRef);
+	 }*/
+#endif	//COREVIDEO_SUPPORTS_METAL // On Mac we have to use IOSurfaceRef for backward compatibility - unless we update MIN_REQUIRED_VERSION to 10.11 we link against an older version of CoreVideo that doesn't support Metal.
+	{
+		FRHIResourceCreateInfo CreateInfo;
+		if (IsMetalPlatform(GMaxRHIShaderPlatform))
+		{
+			// Metal can upload directly from an IOSurface to a 2D texture, so we can just wrap it.
+			CreateInfo.BulkData = new FAvfTexture2DResourceWrapper(Frame);
+		}
+		else
+		{
+			// OpenGL on Mac uploads as a TEXTURE_RECTANGLE for which we have no code, so upload via system memory.
+			CreateInfo.BulkData = new FAvfTexture2DResourceMem(Frame);
+		}
+		CreateInfo.ResourceArray = nullptr;
+		
+		int32 Width = CVPixelBufferGetWidth(Frame);
+		int32 Height = CVPixelBufferGetHeight(Frame);
+		
+		uint32 TexCreateFlags = TexCreate_SRGB;
+		TexCreateFlags |= TexCreate_Dynamic | TexCreate_NoTiling;
+
+		ShaderResource = RHICreateTexture2D(Width, Height, PF_B8G8R8A8, 1, 1, TexCreateFlags | TexCreate_ShaderResource, CreateInfo);
+	}
+	
+	if(ShaderResource.IsValid())
+	{
+		if (VideoSample->Initialize(ShaderResource, Dim, OutputDim, SampleTime, SampleDuration))
+		{
+			Samples.AddVideo(VideoSample);
+		}
+	}
+	
+#else //WITH_ENGINE
+	
+	if(CVPixelBufferLockBaseAddress(Frame, kCVPixelBufferLock_ReadOnly) == kCVReturnSuccess)
+	{
+		uint8* VideoData = (uint8*)CVPixelBufferGetBaseAddress(Frame);
+		
+		if (VideoSample->Initialize(VideoData, Dim, OutputDim, FrameStride, SampleTime, SampleDuration))
+		{
+			Samples.AddVideo(VideoSample);
+		}
+		
+		
+		CVPixelBufferUnlockBaseAddress(Frame, kCVPixelBufferLock_ReadOnly);
+	}
+	
+#endif //WITH_ENGINE
+	
+	CVPixelBufferRelease(Frame);
+}
diff --git a/Engine/Plugins/Media/AvfMedia/Source/AvfMedia/Public/IAvfMediaModule.h b/Engine/Plugins/Media/AvfMedia/Source/AvfMedia/Public/IAvfMediaModule.h
--- a/Engine/Plugins/Media/AvfMedia/Source/AvfMedia/Public/IAvfMediaModule.h
+++ b/Engine/Plugins/Media/AvfMedia/Source/AvfMedia/Public/IAvfMediaModule.h
@@ -2,12 +2,13 @@
 
 #pragma once
 
-#include "CoreTypes.h"
 #include "Templates/SharedPointer.h"
 #include "Modules/ModuleInterface.h"
 
+class IMediaEventSink;
 class IMediaPlayer;
 
+
 /**
  * Interface for the AvfMedia module.
  */
@@ -17,11 +18,12 @@ class IAvfMediaModule
 public:
 
 	/**
-	 * Creates a AV Foundation based media player.
+	 * Create a AV Foundation based media player.
 	 *
+	 * @param EventSink The object that receives media events from the player.
 	 * @return A new media player, or nullptr if a player couldn't be created.
 	 */
-	virtual TSharedPtr<IMediaPlayer, ESPMode::ThreadSafe> CreatePlayer() = 0;
+	virtual TSharedPtr<IMediaPlayer, ESPMode::ThreadSafe> CreatePlayer(IMediaEventSink& EventSink) = 0;
 
 public:
 
diff --git a/Engine/Plugins/Media/AvfMedia/Source/AvfMediaEditor/Private/AvfFileMediaSourceFactory.cpp b/Engine/Plugins/Media/AvfMedia/Source/AvfMediaEditor/Private/AvfFileMediaSourceFactory.cpp
--- a/Engine/Plugins/Media/AvfMedia/Source/AvfMediaEditor/Private/AvfFileMediaSourceFactory.cpp
+++ b/Engine/Plugins/Media/AvfMedia/Source/AvfMediaEditor/Private/AvfFileMediaSourceFactory.cpp
@@ -1,8 +1,11 @@
 // Copyright 1998-2017 Epic Games, Inc. All Rights Reserved.
 
 #include "AvfFileMediaSourceFactory.h"
+
 #include "Misc/Paths.h"
 #include "FileMediaSource.h"
+#include "Containers/UnrealString.h"
+#include "UObject/UObjectGlobals.h"
 
 
 /* UAvfFileMediaSourceFactory structors
diff --git a/Engine/Plugins/Media/AvfMedia/Source/AvfMediaEditor/Private/AvfFileMediaSourceFactory.h b/Engine/Plugins/Media/AvfMedia/Source/AvfMediaEditor/Private/AvfFileMediaSourceFactory.h
--- a/Engine/Plugins/Media/AvfMedia/Source/AvfMediaEditor/Private/AvfFileMediaSourceFactory.h
+++ b/Engine/Plugins/Media/AvfMedia/Source/AvfMediaEditor/Private/AvfFileMediaSourceFactory.h
@@ -2,11 +2,12 @@
 
 #pragma once
 
-#include "CoreMinimal.h"
 #include "UObject/ObjectMacros.h"
 #include "Factories/Factory.h"
+
 #include "AvfFileMediaSourceFactory.generated.h"
 
+
 /**
  * Implements a factory for UFileMediaSource objects.
  */
diff --git a/Engine/Plugins/Media/AvfMedia/Source/AvfMediaFactory/Private/AvfMediaFactoryClasses.cpp b/Engine/Plugins/Media/AvfMedia/Source/AvfMediaFactory/Private/AvfMediaFactoryClasses.cpp
--- a/Engine/Plugins/Media/AvfMedia/Source/AvfMediaFactory/Private/AvfMediaFactoryClasses.cpp
+++ b/Engine/Plugins/Media/AvfMedia/Source/AvfMediaFactory/Private/AvfMediaFactoryClasses.cpp
@@ -1,6 +1,5 @@
 // Copyright 1998-2016 Epic Games, Inc. All Rights Reserved.
 
-#include "CoreMinimal.h"
 #include "AvfMediaSettings.h"
 
 
diff --git a/Engine/Plugins/Media/AvfMedia/Source/AvfMediaFactory/Private/AvfMediaFactoryModule.cpp b/Engine/Plugins/Media/AvfMedia/Source/AvfMediaFactory/Private/AvfMediaFactoryModule.cpp
--- a/Engine/Plugins/Media/AvfMedia/Source/AvfMediaFactory/Private/AvfMediaFactoryModule.cpp
+++ b/Engine/Plugins/Media/AvfMedia/Source/AvfMediaFactory/Private/AvfMediaFactoryModule.cpp
@@ -2,22 +2,27 @@
 
 #include "AvfMediaFactoryPrivate.h"
 
-#include "CoreMinimal.h"
-#include "IAvfMediaModule.h"
+#include "Containers/Array.h"
+#include "Containers/UnrealString.h"
 #include "IMediaModule.h"
 #include "IMediaPlayerFactory.h"
+#include "Internationalization/Internationalization.h"
 #include "Misc/Paths.h"
 #include "Modules/ModuleInterface.h"
 #include "Modules/ModuleManager.h"
-#include "AvfMediaSettings.h"
-#include "UObject/Class.h"
-#include "UObject/WeakObjectPtr.h"
+#include "UObject/NameTypes.h"
 
 #if WITH_EDITOR
 	#include "ISettingsModule.h"
 	#include "ISettingsSection.h"
+	#include "UObject/Class.h"
+	#include "UObject/WeakObjectPtr.h"
 #endif
 
+#include "AvfMediaSettings.h"
+
+#include "../../AvfMedia/Public/IAvfMediaModule.h"
+
 
 DEFINE_LOG_CATEGORY(LogAvfMediaFactory);
 
@@ -35,7 +40,7 @@ public:
 
 	//~ IMediaPlayerFactory interface
 
-	virtual bool CanPlayUrl(const FString& Url, const IMediaOptions& Options, TArray<FText>* OutWarnings, TArray<FText>* OutErrors) const override
+	virtual bool CanPlayUrl(const FString& Url, const IMediaOptions* /*Options*/, TArray<FText>* /*OutWarnings*/, TArray<FText>* OutErrors) const override
 	{
 		FString Scheme;
 		FString Location;
@@ -80,10 +85,10 @@ public:
 		return true;
 	}
 
-	virtual TSharedPtr<IMediaPlayer, ESPMode::ThreadSafe> CreatePlayer() override
+	virtual TSharedPtr<IMediaPlayer, ESPMode::ThreadSafe> CreatePlayer(IMediaEventSink& EventSink) override
 	{
 		auto AvfMediaModule = FModuleManager::LoadModulePtr<IAvfMediaModule>("AvfMedia");
-		return (AvfMediaModule != nullptr) ? AvfMediaModule->CreatePlayer() : nullptr;
+		return (AvfMediaModule != nullptr) ? AvfMediaModule->CreatePlayer(EventSink) : nullptr;
 	}
 
 	virtual FText GetDisplayName() const override
@@ -102,6 +107,16 @@ public:
 		return SupportedPlatforms;
 	}
 
+	virtual bool SupportsFeature(EMediaFeature Feature) const override
+	{
+		return ((Feature == EMediaFeature::AudioSamples) ||
+				(Feature == EMediaFeature::AudioTracks) ||
+				(Feature == EMediaFeature::CaptionTracks) ||
+				(Feature == EMediaFeature::OverlaySamples) ||
+				(Feature == EMediaFeature::VideoSamples) ||
+				(Feature == EMediaFeature::VideoTracks));
+	}
+
 public:
 
 	//~ IModuleInterface interface
diff --git a/Engine/Plugins/Media/AvfMedia/Source/AvfMediaFactory/Private/AvfMediaFactoryPrivate.h b/Engine/Plugins/Media/AvfMedia/Source/AvfMediaFactory/Private/AvfMediaFactoryPrivate.h
--- a/Engine/Plugins/Media/AvfMedia/Source/AvfMediaFactory/Private/AvfMediaFactoryPrivate.h
+++ b/Engine/Plugins/Media/AvfMedia/Source/AvfMediaFactory/Private/AvfMediaFactoryPrivate.h
@@ -2,9 +2,8 @@
 
 #pragma once
 
-#include "CoreMinimal.h"
+#include "Logging/LogMacros.h"
 
 
-
-/** Log category for the AvfMediaFactory module. */
+/** Log category for the AvfMedia module. */
 DECLARE_LOG_CATEGORY_EXTERN(LogAvfMediaFactory, Log, All);
diff --git a/Engine/Plugins/Media/AvfMedia/Source/AvfMediaFactory/Public/AvfMediaSettings.h b/Engine/Plugins/Media/AvfMedia/Source/AvfMediaFactory/Public/AvfMediaSettings.h
--- a/Engine/Plugins/Media/AvfMedia/Source/AvfMediaFactory/Public/AvfMediaSettings.h
+++ b/Engine/Plugins/Media/AvfMedia/Source/AvfMediaFactory/Public/AvfMediaSettings.h
@@ -2,12 +2,15 @@
 
 #pragma once
 
-#include "CoreMinimal.h"
 #include "UObject/ObjectMacros.h"
 #include "UObject/Object.h"
+
 #include "AvfMediaSettings.generated.h"
 
 
+/**
+ * Settings for the AvfMedia plug-in.
+ */
 UCLASS(config=Engine)
 class AVFMEDIAFACTORY_API UAvfMediaSettings
 	: public UObject
diff --git a/Engine/Plugins/Media/ImgMedia/Source/ImgMedia/ImgMedia.Build.cs b/Engine/Plugins/Media/ImgMedia/Source/ImgMedia/ImgMedia.Build.cs
new file mode 100644
--- /dev/null
+++ b/Engine/Plugins/Media/ImgMedia/Source/ImgMedia/ImgMedia.Build.cs
@@ -0,0 +1,53 @@
+// Copyright 1998-2016 Epic Games, Inc. All Rights Reserved.
+
+namespace UnrealBuildTool.Rules
+{
+	public class ImgMedia : ModuleRules
+	{
+		public ImgMedia(ReadOnlyTargetRules Target) : base(Target)
+		{
+			DynamicallyLoadedModuleNames.AddRange(
+				new string[] {
+					"ImageWrapper",
+					"Media",
+				});
+
+			PrivateDependencyModuleNames.AddRange(
+				new string[] {
+					"Core",
+					"CoreUObject",
+					"Engine",
+					"ImgMediaFactory",
+					"RenderCore",
+					"RHI",
+				});
+
+			PrivateIncludePathModuleNames.AddRange(
+				new string[] {
+					"ImageWrapper",
+					"Media",
+				});
+
+			PrivateIncludePaths.AddRange(
+				new string[] {
+					"ImgMedia/Private",
+					"ImgMedia/Private/Assets",
+					"ImgMedia/Private/Loader",
+					"ImgMedia/Private/Player",
+					"ImgMedia/Private/Readers",
+				});
+
+			PublicDependencyModuleNames.AddRange(
+				new string[] {
+					"MediaAssets",
+				});
+
+			if ((Target.Platform == UnrealTargetPlatform.Mac) ||
+				(Target.Platform == UnrealTargetPlatform.Win32) ||
+				(Target.Platform == UnrealTargetPlatform.Win64))
+			{
+				PrivateDependencyModuleNames.Add("OpenExrWrapper");
+			}
+		}
+	}
+}
diff --git a/Engine/Plugins/Media/ImgMedia/Source/ImgMedia/Private/Assets/ImgMediaSource.cpp b/Engine/Plugins/Media/ImgMedia/Source/ImgMedia/Private/Assets/ImgMediaSource.cpp
new file mode 100644
--- /dev/null
+++ b/Engine/Plugins/Media/ImgMedia/Source/ImgMedia/Private/Assets/ImgMediaSource.cpp
@@ -0,0 +1,119 @@
+// Copyright 1998-2016 Epic Games, Inc. All Rights Reserved.
+
+#include "ImgMediaSource.h"
+#include "ImgMediaPrivate.h"
+
+#include "HAL/FileManager.h"
+#include "Misc/Paths.h"
+
+
+/* UImgMediaSource structors
+ *****************************************************************************/
+
+UImgMediaSource::UImgMediaSource()
+	: FramesPerSecondOverride(0.0f)
+{ }
+
+
+/* UImgMediaSource interface
+ *****************************************************************************/
+
+void UImgMediaSource::GetProxies(TArray<FString>& OutProxies) const
+{
+	IFileManager::Get().FindFiles(OutProxies, *FPaths::Combine(GetFullPath(), TEXT("*")), false, true);
+}
+
+
+void UImgMediaSource::SetSequencePath(const FString& Path)
+{
+	const FString SanitizedPath = FPaths::GetPath(Path);
+
+	if (SanitizedPath.IsEmpty() || SanitizedPath.StartsWith(TEXT(".")))
+	{
+		SequencePath.Path = SanitizedPath;
+	}
+	else
+	{
+		FString FullPath = FPaths::ConvertRelativePathToFull(SanitizedPath);
+		const FString FullGameContentDir = FPaths::ConvertRelativePathToFull(FPaths::ProjectContentDir());
+
+		if (FullPath.StartsWith(FullGameContentDir))
+		{
+			FPaths::MakePathRelativeTo(FullPath, *FullGameContentDir);
+			FullPath = FString(TEXT("./")) + FullPath;
+		}
+
+		SequencePath.Path = FullPath;
+	}
+}
+
+
+/* IMediaOptions interface
+ *****************************************************************************/
+
+double UImgMediaSource::GetMediaOption(const FName& Key, const double DefaultValue) const
+{
+	if (Key == ImgMedia::FramesPerSecondOverrideOption)
+	{
+		return FramesPerSecondOverride;
+	}
+
+	return Super::GetMediaOption(Key, DefaultValue);
+}
+
+
+FString UImgMediaSource::GetMediaOption(const FName& Key, const FString& DefaultValue) const
+{
+	if (Key == ImgMedia::ProxyOverrideOption)
+	{
+		return ProxyOverride;
+	}
+
+	return Super::GetMediaOption(Key, DefaultValue);
+}
+
+
+bool UImgMediaSource::HasMediaOption(const FName& Key) const
+{
+	if ((Key == ImgMedia::FramesPerSecondOverrideOption) ||
+		(Key == ImgMedia::ProxyOverrideOption))
+	{
+		return true;
+	}
+
+	return Super::HasMediaOption(Key);
+}
+
+
+/* UMediaSource interface
+ *****************************************************************************/
+
+FString UImgMediaSource::GetUrl() const
+{
+	return FString(TEXT("img://")) + GetFullPath();
+}
+
+
+bool UImgMediaSource::Validate() const
+{
+	return FPaths::DirectoryExists(GetFullPath());
+}
+
+
+/* UFileMediaSource implementation
+ *****************************************************************************/
+
+FString UImgMediaSource::GetFullPath() const
+{
+	if (!FPaths::IsRelative(SequencePath.Path))
+	{
+		return SequencePath.Path;
+	}
+
+	if (SequencePath.Path.StartsWith(TEXT("./")))
+	{
+		return FPaths::ConvertRelativePathToFull(FPaths::ProjectContentDir(), SequencePath.Path.RightChop(2));
+	}
+
+	return FPaths::ConvertRelativePathToFull(SequencePath.Path);
+}
diff --git a/Engine/Plugins/Media/ImgMedia/Source/ImgMedia/Private/ImgMediaModule.cpp b/Engine/Plugins/Media/ImgMedia/Source/ImgMedia/Private/ImgMediaModule.cpp
new file mode 100644
--- /dev/null
+++ b/Engine/Plugins/Media/ImgMedia/Source/ImgMedia/Private/ImgMediaModule.cpp
@@ -0,0 +1,43 @@
+// Copyright 1998-2016 Epic Games, Inc. All Rights Reserved.
+
+#include "ImgMediaPrivate.h"
+
+#include "Modules/ModuleManager.h"
+
+#include "ImgMediaPlayer.h"
+#include "IImgMediaModule.h"
+
+
+DEFINE_LOG_CATEGORY(LogImgMedia);
+
+
+/**
+ * Implements the AVFMedia module.
+ */
+class FImgMediaModule
+	: public IImgMediaModule
+{
+public:
+
+	/** Default constructor. */
+	FImgMediaModule() { }
+
+public:
+
+	//~ IImgMediaModule interface
+
+	virtual TSharedPtr<IMediaPlayer, ESPMode::ThreadSafe> CreatePlayer(IMediaEventSink& EventSink) override
+	{
+		return MakeShared<FImgMediaPlayer, ESPMode::ThreadSafe>(EventSink);
+	}
+
+public:
+
+	//~ IModuleInterface interface
+
+	virtual void StartupModule() override { }
+	virtual void ShutdownModule() override { }
+};
+
+
+IMPLEMENT_MODULE(FImgMediaModule, ImgMedia);
diff --git a/Engine/Plugins/Media/ImgMedia/Source/ImgMedia/Private/ImgMediaPrivate.h b/Engine/Plugins/Media/ImgMedia/Source/ImgMedia/Private/ImgMediaPrivate.h
new file mode 100644
--- /dev/null
+++ b/Engine/Plugins/Media/ImgMedia/Source/ImgMedia/Private/ImgMediaPrivate.h
@@ -0,0 +1,30 @@
+// Copyright 1998-2016 Epic Games, Inc. All Rights Reserved.
+
+#pragma once
+
+#include "ImgMediaSettings.h"
+#include "Logging/LogMacros.h"
+#include "UObject/NameTypes.h"
+
+
+/** The OpenEXR image reader is supported on macOS and Windows only. */
+#define IMGMEDIA_EXR_SUPPORTED_PLATFORM (PLATFORM_MAC || PLATFORM_WINDOWS)
+
+/** Default number of frames per second for image sequences. */
+#define IMGMEDIA_DEFAULT_FPS 24.0
+
+/** Log category for the this module. */
+DECLARE_LOG_CATEGORY_EXTERN(LogImgMedia, Log, All);
+
+
+namespace ImgMedia
+{
+	/** Name of the FramesPerSecondAttribute media option. */
+	static FName FramesPerSecondAttributeOption("FramesPerSecondAttribute");
+
+	/** Name of the FramesPerSecondOverride media option. */
+	static FName FramesPerSecondOverrideOption("FramesPerSecondOverride");
+
+	/** Name of the ProxyOverride media option. */
+	static FName ProxyOverrideOption("ProxyOverride");
+}
diff --git a/Engine/Plugins/Media/ImgMedia/Source/ImgMedia/Private/Loader/IImgMediaLoader.h b/Engine/Plugins/Media/ImgMedia/Source/ImgMedia/Private/Loader/IImgMediaLoader.h
new file mode 100644
--- /dev/null
+++ b/Engine/Plugins/Media/ImgMedia/Source/ImgMedia/Private/Loader/IImgMediaLoader.h
@@ -0,0 +1,27 @@
+// Copyright 1998-2016 Epic Games, Inc. All Rights Reserved.
+
+#pragma once
+
+#include "CoreTypes.h"
+#include "Templates/SharedPointer.h"
+
+class FImgMediaLoaderWork;
+struct FImgMediaFrame;
+
+
+/**
+ * Interface for the image sequence loader.
+ */
+class IImgMediaLoader
+{
+public:
+
+	/**
+	 * Notify the loader that a work item completed.
+	 *
+	 * @param CompletedWork The work item that completed.
+	 * @param FrameNumber Number of the frame that was read.
+	 * @param Frame The frame that was read, or nullptr if reading failed.
+	 */
+	virtual void NotifyWorkComplete(FImgMediaLoaderWork& CompletedWork, int32 FrameNumber, const TSharedPtr<FImgMediaFrame, ESPMode::ThreadSafe>& Frame) = 0;
+};
