Fix for bug Playing video in the Editor's Media Player on a Mac results in a degradation in quality with some noticeable banding present.  Select correct pixel formats for the encoded video in H.264.  Use correct colour transform matrix.  Don't use bilinear sampling of the luma texture - it should be 1:1 with the output.

diff --git a/Engine/Plugins/Media/AvfMedia/Source/AvfMedia/Private/Player/AvfMediaTracks.cpp b/Engine/Plugins/Media/AvfMedia/Source/AvfMedia/Private/Player/AvfMediaTracks.cpp
--- a/Engine/Plugins/Media/AvfMedia/Source/AvfMedia/Private/Player/AvfMediaTracks.cpp
+++ b/Engine/Plugins/Media/AvfMedia/Source/AvfMedia/Private/Player/AvfMediaTracks.cpp
@@ -20,6 +20,14 @@
 
 NS_ASSUME_NONNULL_BEGIN
 
+static int32 AVFMediaForceDecodeBGRA = 0;
+FAutoConsoleVariableRef CVarAVFMediaForceDecodeBGRA(
+	TEXT("m.avf.ForceDecodeBGRA"),
+	AVFMediaForceDecodeBGRA,
+	TEXT("Change between YUV decode and convert to BGRA in UE4 Shader (Keeps everything on the GPU) or always force Apple framework to perform the decode to BRGA (potential performance penalty).\n")
+	TEXT("0: Auto Detect YUV (Default) and decode to BGRA in UE4 Shader, 1: Force AV Framework to decode to BGRA"),
+	ECVF_ReadOnly);
+
 /* FAVPlayerItemLegibleOutputPushDelegate
  *****************************************************************************/
 
@@ -436,32 +444,49 @@ void FAvfMediaTracks::Initialize(AVPlayerItem* InPlayerItem, FString& OutInfo)
 		}
 		else if ([MediaType isEqualToString:AVMediaTypeVideo])
 		{
+			check(COREVIDEO_SUPPORTS_METAL);
+			
 			NSMutableDictionary* OutputSettings = [NSMutableDictionary dictionary];
-			// Mac:
-			// On Mac kCVPixelFormatType_422YpCbCr8 is the preferred single-plane YUV format but for H.264 bi-planar formats are the optimal choice
-			// The native RGBA format is 32ARGB but we use 32BGRA for consistency with iOS for now.
-			//
-			// iOS/tvOS:
-			// On iOS only bi-planar kCVPixelFormatType_420YpCbCr8BiPlanarVideoRange/kCVPixelFormatType_420YpCbCr8BiPlanarFullRange are supported for YUV so an additional conversion is required.
-			// The only RGBA format is 32BGRA
-#if COREVIDEO_SUPPORTS_METAL
-			[OutputSettings setObject : [NSNumber numberWithInt : kCVPixelFormatType_420YpCbCr8BiPlanarFullRange] forKey : (NSString*)kCVPixelBufferPixelFormatTypeKey];
-#else
-			[OutputSettings setObject : [NSNumber numberWithInt : kCVPixelFormatType_32BGRA] forKey : (NSString*)kCVPixelBufferPixelFormatTypeKey];
-#endif
-
-#if WITH_ENGINE
-			// Setup sharing with RHI's starting with the optional Metal RHI
-			if (FPlatformMisc::HasPlatformFeature(TEXT("Metal")))
+			CMFormatDescriptionRef DescRef = (CMFormatDescriptionRef)[AssetTrack.formatDescriptions objectAtIndex:0];
+			
+			// Select decode format
+			int32 DecodePixelBufferFormat = 0;
+	
+			if(!AVFMediaForceDecodeBGRA)
 			{
-				[OutputSettings setObject:[NSNumber numberWithBool:YES] forKey:(NSString*)kCVPixelBufferMetalCompatibilityKey];
+				CFDictionaryRef FormatExtensions = CMFormatDescriptionGetExtensions(DescRef);
+				if(FormatExtensions)
+				{
+					CFStringRef FormatName = (CFStringRef)CFDictionaryGetValue(FormatExtensions, kCMFormatDescriptionExtension_FormatName);
+					
+					// kCMVideoCodecType_H264 : 'avc1' - why is this a CFString in the dictionary and not a CFNumber?
+					if(FormatName && kCFCompareEqualTo == CFStringCompare(FormatName, CFSTR("'avc1'"), 0))
+					{
+						CFBooleanRef bFullRange = (CFBooleanRef)CFDictionaryGetValue(FormatExtensions, kCMFormatDescriptionExtension_FullRangeVideo);
+						if(bFullRange && (bool)CFBooleanGetValue(bFullRange))
+						{
+							DecodePixelBufferFormat = kCVPixelFormatType_420YpCbCr8BiPlanarFullRange;
+						}
+						else
+						{
+							DecodePixelBufferFormat = kCVPixelFormatType_420YpCbCr8BiPlanarVideoRange;
+						}
+					}
+				}
 			}
 
-#if PLATFORM_MAC
-			[OutputSettings setObject:[NSNumber numberWithBool:YES] forKey:(NSString*)kCVPixelBufferOpenGLCompatibilityKey];
-#else
-			[OutputSettings setObject:[NSNumber numberWithBool:YES] forKey:(NSString*)kCVPixelBufferOpenGLESCompatibilityKey];
-#endif
+			// BGRA32 is the fallback
+			if(DecodePixelBufferFormat == 0)
+			{
+				DecodePixelBufferFormat = kCVPixelFormatType_32BGRA;
+			}
+			
+			[OutputSettings setObject : [NSNumber numberWithInt :  DecodePixelBufferFormat] forKey : (NSString*)kCVPixelBufferPixelFormatTypeKey];
+			
+#if WITH_ENGINE
+			
+			[OutputSettings setObject:[NSNumber numberWithBool:YES] forKey:(NSString*)kCVPixelBufferMetalCompatibilityKey];
+			
 #endif //WITH_ENGINE
 
 			// Use unaligned rows
@@ -479,6 +504,7 @@ void FAvfMediaTracks::Initialize(AVPlayerItem* InPlayerItem, FString& OutInfo)
 			Track->Name = FString::Printf(TEXT("Video Track %i"), TrackIndex);
 			Track->Output = Output;
 			Track->Loaded = true;
+			Track->bFullRangeVideo = DecodePixelBufferFormat != kCVPixelFormatType_420YpCbCr8BiPlanarVideoRange;
 			
 			// NominalFrameRate can be zero (e.g HLS streams) - try again using min frame duration, otherwise it's unknown - possibly variable - use a default
 			Track->FrameRate = AssetTrack.nominalFrameRate;
@@ -498,7 +524,6 @@ void FAvfMediaTracks::Initialize(AVPlayerItem* InPlayerItem, FString& OutInfo)
 			CGSize NaturalFrameSize = AssetTrack.naturalSize;
 			Track->FrameSize = FIntPoint((int32)NaturalFrameSize.width, (int32)NaturalFrameSize.height);
 
-			CMFormatDescriptionRef DescRef = (CMFormatDescriptionRef)[AssetTrack.formatDescriptions objectAtIndex:0];
 			CMVideoCodecType CodecType = CMFormatDescriptionGetMediaSubType(DescRef);
 			OutInfo += FString::Printf(TEXT("    Codec: %s\n"), *AvfMedia::CodecTypeToString(CodecType));
 			OutInfo += FString::Printf(TEXT("    Dimensions: %i x %i\n"), (int32)NaturalFrameSize.width, (int32)NaturalFrameSize.height);
@@ -1007,11 +1032,13 @@ bool FAvfMediaTracks::SelectTrack(EMediaTrackType TrackType, int32 TrackIndex)
 					AVPlayerItemVideoOutput* Output;
 					TWeakPtr<FAvfMediaVideoSampler, ESPMode::ThreadSafe> VideoSamplerPtr;
 					float FrameRate;
+					bool bFullRange;
 				}
 				Params = {
 					(AVPlayerItemVideoOutput*)VideoTracks[SelectedVideoTrack].Output,
 					VideoSampler,
-					VideoTracks[TrackIndex].FrameRate
+					VideoTracks[TrackIndex].FrameRate,
+					VideoTracks[TrackIndex].bFullRangeVideo
 				};
 
 				ENQUEUE_RENDER_COMMAND(AvfMediaVideoSamplerSetOutput)(
@@ -1021,7 +1048,7 @@ bool FAvfMediaTracks::SelectTrack(EMediaTrackType TrackType, int32 TrackIndex)
 
 						if (PinnedVideoSampler.IsValid())
 						{
-							PinnedVideoSampler->SetOutput(Params.Output, Params.FrameRate);
+							PinnedVideoSampler->SetOutput(Params.Output, Params.FrameRate, Params.bFullRange);
 						}
 					});
 			}
diff --git a/Engine/Plugins/Media/AvfMedia/Source/AvfMedia/Private/Player/AvfMediaTracks.h b/Engine/Plugins/Media/AvfMedia/Source/AvfMedia/Private/Player/AvfMediaTracks.h
--- a/Engine/Plugins/Media/AvfMedia/Source/AvfMedia/Private/Player/AvfMediaTracks.h
+++ b/Engine/Plugins/Media/AvfMedia/Source/AvfMedia/Private/Player/AvfMediaTracks.h
@@ -40,6 +40,7 @@ class FAvfMediaTracks
 		: Loaded(false)
 		, FrameSize(0, 0)
 		, FrameRate(0.f)
+		, bFullRangeVideo(false)
 		{ }
 		
 		AVAssetTrack* AssetTrack;
@@ -52,6 +53,7 @@ class FAvfMediaTracks
 		//Cached Video Track Data
 		FIntPoint FrameSize;
 		float FrameRate;
+		bool bFullRangeVideo;
 	};
 
 public:
diff --git a/Engine/Plugins/Media/AvfMedia/Source/AvfMedia/Private/Player/AvfMediaVideoSampler.cpp b/Engine/Plugins/Media/AvfMedia/Source/AvfMedia/Private/Player/AvfMediaVideoSampler.cpp
--- a/Engine/Plugins/Media/AvfMedia/Source/AvfMedia/Private/Player/AvfMediaVideoSampler.cpp
+++ b/Engine/Plugins/Media/AvfMedia/Source/AvfMedia/Private/Player/AvfMediaVideoSampler.cpp
@@ -132,7 +132,8 @@ FAvfMediaVideoSampler::FAvfMediaVideoSampler(FMediaSamples& InSamples)
 	, Samples(InSamples)
 	, VideoSamplePool(new FAvfMediaTextureSamplePool)
 	, FrameDuration(0.0f)
-#if WITH_ENGINE && COREVIDEO_SUPPORTS_METAL
+	, ColorTransform(nullptr)
+#if WITH_ENGINE
 	, MetalTextureCache(nullptr)
 #endif
 { }
@@ -145,7 +146,7 @@ FAvfMediaVideoSampler::~FAvfMediaVideoSampler()
 	delete VideoSamplePool;
 	VideoSamplePool = nullptr;
 
-#if WITH_ENGINE && COREVIDEO_SUPPORTS_METAL
+#if WITH_ENGINE
 	if (MetalTextureCache)
 	{
 		CFRelease(MetalTextureCache);
@@ -158,7 +159,7 @@ FAvfMediaVideoSampler::~FAvfMediaVideoSampler()
 /* FAvfMediaVideoSampler interface
  *****************************************************************************/
 
-void FAvfMediaVideoSampler::SetOutput(AVPlayerItemVideoOutput* InOutput, float InFrameRate)
+void FAvfMediaVideoSampler::SetOutput(AVPlayerItemVideoOutput* InOutput, float InFrameRate, bool bFullRange)
 {
 	check(IsInRenderingThread());
 
@@ -169,6 +170,15 @@ void FAvfMediaVideoSampler::SetOutput(AVPlayerItemVideoOutput* InOutput, float I
 	Output = InOutput;
 
 	FrameDuration = 1.f / InFrameRate;
+	
+	if(bFullRange)
+	{
+		ColorTransform = &MediaShaders::YuvToRgbRec709Full;
+	}
+	else
+	{
+		ColorTransform = &MediaShaders::YuvToRgbRec709;
+	}
 }
 
 
@@ -211,9 +221,10 @@ void FAvfMediaVideoSampler::Tick()
 
 #if WITH_ENGINE
 	TRefCountPtr<FRHITexture2D> ShaderResource;
-#if COREVIDEO_SUPPORTS_METAL
-	// On iOS/tvOS we use the Metal texture cache
-	if (IsMetalPlatform(GMaxRHIShaderPlatform))
+
+	// We have to support Metal for this object now
+	check(COREVIDEO_SUPPORTS_METAL);
+	check(IsMetalPlatform(GMaxRHIShaderPlatform));
 	{
 		if (!MetalTextureCache)
 		{
@@ -227,7 +238,8 @@ void FAvfMediaVideoSampler::Tick()
 		
 		if (CVPixelBufferIsPlanar(Frame))
 		{
-			check(IsMetalPlatform(GMaxRHIShaderPlatform));
+			// Expecting BiPlanar kCVPixelFormatType_420YpCbCr8BiPlanar Full/Video
+			check(CVPixelBufferGetPlaneCount(Frame) == 2);
 			
 			uint32 TexCreateFlags = TexCreate_Dynamic | TexCreate_NoTiling;
 			
@@ -260,7 +272,7 @@ void FAvfMediaVideoSampler::Tick()
 			TRefCountPtr<FRHITexture2D> UVTex = RHICreateTexture2D(UVWidth, UVHeight, PF_R8G8, 1, 1, TexCreateFlags | TexCreate_ShaderResource, UVCreateInfo);
 			
 			FRHIResourceCreateInfo Info;
-			ShaderResource = RHICreateTexture2D(YWidth, YHeight, PF_B8G8R8A8, 1, 1, TexCreateFlags | TexCreate_ShaderResource | TexCreate_RenderTargetable, Info);
+			ShaderResource = RHICreateTexture2D(YWidth, YHeight, PF_B8G8R8A8, 1, 1, TexCreateFlags | TexCreate_ShaderResource | TexCreate_RenderTargetable | TexCreate_SRGB, Info);
 			
 			// render video frame into sink texture
 			FRHICommandListImmediate& RHICmdList = FRHICommandListExecutor::GetImmediateCommandList();
@@ -287,7 +299,7 @@ void FAvfMediaVideoSampler::Tick()
 
 					SetGraphicsPipelineState(RHICmdList, GraphicsPSOInit);
 
-					PixelShader->SetParameters(RHICmdList, YTex, UVTex, MediaShaders::YuvToSrgbPs4, MediaShaders::YUVOffset8bits, true);
+					PixelShader->SetParameters(RHICmdList, YTex, UVTex, *ColorTransform, MediaShaders::YUVOffset8bits, true);
 
 					FVertexBufferRHIRef VertexBuffer = CreateTempMediaVertexBuffer();
 					RHICmdList.SetStreamSource(0, VertexBuffer, 0);
@@ -324,74 +336,6 @@ void FAvfMediaVideoSampler::Tick()
 			CFRelease(TextureRef);
 		}
 	}
-	else // Ran out of time to implement efficient OpenGLES texture upload - its running out of memory.
-	/*{
-	 if (!OpenGLTextureCache)
-	 {
-	 EAGLContext* Context = [EAGLContext currentContext];
-	 check(Context);
-	 
-	 CVReturn Return = CVOpenGLESTextureCacheCreate(kCFAllocatorDefault, nullptr, Context, nullptr, &OpenGLTextureCache);
-	 check(Return == kCVReturnSuccess);
-	 }
-	 check(OpenGLTextureCache);
-	 
-	 int32 Width = CVPixelBufferGetWidth(Frame);
-	 int32 Height = CVPixelBufferGetHeight(Frame);
-	 
-	 CVOpenGLESTextureRef TextureRef = nullptr;
-	 CVReturn Result = CVOpenGLESTextureCacheCreateTextureFromImage(kCFAllocatorDefault, OpenGLTextureCache, Frame, nullptr, GL_TEXTURE_2D, GL_RGBA, Width, Height, GL_RGBA, GL_UNSIGNED_BYTE, 0, &TextureRef);
-	 check(Result == kCVReturnSuccess);
-	 check(TextureRef);
-	 
-	 FRHIResourceCreateInfo CreateInfo;
-	 CreateInfo.BulkData = new FAvfTexture2DResourceWrapper(TextureRef);
-	 CreateInfo.ResourceArray = nullptr;
-	 
-	 uint32 TexCreateFlags = 0;
-	 TexCreateFlags |= TexCreate_Dynamic | TexCreate_NoTiling;
-	 
-	 TRefCountPtr<FRHITexture2D> RenderTarget;
-	 TRefCountPtr<FRHITexture2D> ShaderResource;
-	 
-	 RHICreateTargetableShaderResource2D(Width,
-	 Height,
-	 PF_R8G8B8A8,
-	 1,
-	 TexCreateFlags,
-	 TexCreate_RenderTargetable,
-	 false,
-	 CreateInfo,
-	 RenderTarget,
-	 ShaderResource
-	 );
-	 
-	 VideoSink->UpdateTextureSinkResource(RenderTarget, ShaderResource);
-	 CFRelease(TextureRef);
-	 }*/
-#endif	//COREVIDEO_SUPPORTS_METAL // On Mac we have to use IOSurfaceRef for backward compatibility - unless we update MIN_REQUIRED_VERSION to 10.11 we link against an older version of CoreVideo that doesn't support Metal.
-	{
-		FRHIResourceCreateInfo CreateInfo;
-		if (IsMetalPlatform(GMaxRHIShaderPlatform))
-		{
-			// Metal can upload directly from an IOSurface to a 2D texture, so we can just wrap it.
-			CreateInfo.BulkData = new FAvfTexture2DResourceWrapper(Frame);
-		}
-		else
-		{
-			// OpenGL on Mac uploads as a TEXTURE_RECTANGLE for which we have no code, so upload via system memory.
-			CreateInfo.BulkData = new FAvfTexture2DResourceMem(Frame);
-		}
-		CreateInfo.ResourceArray = nullptr;
-		
-		int32 Width = CVPixelBufferGetWidth(Frame);
-		int32 Height = CVPixelBufferGetHeight(Frame);
-		
-		uint32 TexCreateFlags = TexCreate_SRGB;
-		TexCreateFlags |= TexCreate_Dynamic | TexCreate_NoTiling;
-
-		ShaderResource = RHICreateTexture2D(Width, Height, PF_B8G8R8A8, 1, 1, TexCreateFlags | TexCreate_ShaderResource, CreateInfo);
-	}
 	
 	if(ShaderResource.IsValid())
 	{
diff --git a/Engine/Plugins/Media/AvfMedia/Source/AvfMedia/Private/Player/AvfMediaVideoSampler.h b/Engine/Plugins/Media/AvfMedia/Source/AvfMedia/Private/Player/AvfMediaVideoSampler.h
--- a/Engine/Plugins/Media/AvfMedia/Source/AvfMedia/Private/Player/AvfMediaVideoSampler.h
+++ b/Engine/Plugins/Media/AvfMedia/Source/AvfMedia/Private/Player/AvfMediaVideoSampler.h
@@ -33,7 +33,7 @@ public:
 	 *
 	 * @param Output The output object.
 	 */
-	void SetOutput(AVPlayerItemVideoOutput* Output, float InFrameRate);
+	void SetOutput(AVPlayerItemVideoOutput* Output, float InFrameRate, bool bFullRange = false);
 
 	/** Tick the video sampler (on the render thread). */
 	void Tick();
@@ -52,9 +52,11 @@ private:
 	/** Video sample object pool. */
 	FAvfMediaTextureSamplePool* VideoSamplePool;
 
+	/** Current Video cached decode information */
 	float FrameDuration;
+	FMatrix const* ColorTransform;
 
-#if WITH_ENGINE && COREVIDEO_SUPPORTS_METAL
+#if WITH_ENGINE
 
 	/** The Metal texture cache for unbuffered texture uploads. */
 	CVMetalTextureCacheRef MetalTextureCache;
diff --git a/Engine/Source/Runtime/UtilityShaders/Private/MediaShaders.cpp b/Engine/Source/Runtime/UtilityShaders/Private/MediaShaders.cpp
--- a/Engine/Source/Runtime/UtilityShaders/Private/MediaShaders.cpp
+++ b/Engine/Source/Runtime/UtilityShaders/Private/MediaShaders.cpp
@@ -265,10 +265,13 @@ void FYCbCrConvertPS::SetParameters(FRHICommandList& CommandList, TRefCountPtr<F
 {
 	FYCbCrConvertUB UB;
 	{
+		// Chroma is not usually 1:1 with the output textxure
 		UB.CbCrSampler = TStaticSamplerState<SF_Bilinear>::GetRHI();
 		UB.CbCrTexture = CbCrTexture;
 		UB.ColorTransform = MediaShaders::CombineColorTransformAndOffset(ColorTransform, YUVOffset);
-		UB.LumaSampler = TStaticSamplerState<SF_Bilinear>::GetRHI();
+		
+		// Luma should be 1:1 with the output texture and needs to be point sampled
+		UB.LumaSampler = TStaticSamplerState<SF_Point>::GetRHI();
 		UB.LumaTexture = LumaTexture;
 		UB.SrgbToLinear = SrgbToLinear;
 	}
