Underlying fix for Media current time > Duration while waiting for end of media event (Thanks to [REDACTED_PERSON] for highlighting the issue).

diff --git a/Engine/Plugins/Media/AvfMedia/Source/AvfMedia/Private/Player/AvfMediaPlayer.cpp b/Engine/Plugins/Media/AvfMedia/Source/AvfMedia/Private/Player/AvfMediaPlayer.cpp
--- a/Engine/Plugins/Media/AvfMedia/Source/AvfMedia/Private/Player/AvfMediaPlayer.cpp
+++ b/Engine/Plugins/Media/AvfMedia/Source/AvfMedia/Private/Player/AvfMediaPlayer.cpp
@@ -831,7 +831,8 @@ void FAvfMediaPlayer::TickInput(FTimespan DeltaTime, FTimespan /*Timecode*/)
 					}
 					else
 					{
-						CurrentTime += DeltaTime * CurrentRate;
+						// On certain movies (short ones it seems) the AVPlayerItemDidPlayToEndTimeNotification can be delayed - clamp the time to the duration
+						CurrentTime = FMath::Clamp(CurrentTime + (DeltaTime * CurrentRate), FTimespan::Zero(), Duration);
 					}
 				}
 				break;
