Merge of CL 5808963 - support streaming live video on iOS

diff --git a/Engine/Plugins/Media/AvfMedia/Source/AvfMedia/Private/Player/AvfMediaVideoSampler.cpp b/Engine/Plugins/Media/AvfMedia/Source/AvfMedia/Private/Player/AvfMediaVideoSampler.cpp
--- a/Engine/Plugins/Media/AvfMedia/Source/AvfMedia/Private/Player/AvfMediaVideoSampler.cpp
+++ b/Engine/Plugins/Media/AvfMedia/Source/AvfMedia/Private/Player/AvfMediaVideoSampler.cpp
@@ -210,6 +210,13 @@ void FAvfMediaVideoSampler::Tick()
 	const FTimespan SampleDuration = FTimespan::FromSeconds(FrameDuration);
 	const FTimespan SampleTime = FTimespan::FromSeconds(CMTimeGetSeconds(OutputItemTime));
 
+	ProcessFrame(Frame, SampleTime, SampleDuration);
+}
+
+void FAvfMediaVideoSampler::ProcessFrame(CVPixelBufferRef Frame, FTimespan SampleTime, FTimespan SampleDuration)
+{
+	check(IsInRenderingThread());
+
 	const int32 FrameHeight = CVPixelBufferGetHeight(Frame);
 	const int32 FrameWidth = CVPixelBufferGetWidth(Frame);
 	const int32 FrameStride = CVPixelBufferGetBytesPerRow(Frame);
@@ -341,7 +348,7 @@ void FAvfMediaVideoSampler::Tick()
 	{
 		if (VideoSample->Initialize(ShaderResource, Dim, OutputDim, SampleTime, SampleDuration))
 		{
-			Samples.AddVideo(VideoSample);
+			ProcessOutputSample(VideoSample);
 		}
 	}
 	
@@ -353,7 +360,7 @@ void FAvfMediaVideoSampler::Tick()
 		
 		if (VideoSample->Initialize(VideoData, Dim, OutputDim, FrameStride, SampleTime, SampleDuration))
 		{
-			Samples.AddVideo(VideoSample);
+			ProcessOutputSample(VideoSample);
 		}
 		
 		
@@ -364,3 +371,8 @@ void FAvfMediaVideoSampler::Tick()
 	
 	CVPixelBufferRelease(Frame);
 }
+
+void FAvfMediaVideoSampler::ProcessOutputSample(const TSharedRef<IMediaTextureSample, ESPMode::ThreadSafe>& Sample)
+{
+	Samples.AddVideo(Sample);
+}
diff --git a/Engine/Plugins/Media/AvfMedia/Source/AvfMedia/Private/Player/AvfMediaVideoSampler.h b/Engine/Plugins/Media/AvfMedia/Source/AvfMedia/Private/Player/AvfMediaVideoSampler.h
--- a/Engine/Plugins/Media/AvfMedia/Source/AvfMedia/Private/Player/AvfMediaVideoSampler.h
+++ b/Engine/Plugins/Media/AvfMedia/Source/AvfMedia/Private/Player/AvfMediaVideoSampler.h
@@ -38,6 +38,11 @@ public:
 	/** Tick the video sampler (on the render thread). */
 	void Tick();
 
+	void ProcessFrame(CVPixelBufferRef Frame, FTimespan SampleTime, FTimespan SampleDuration);
+
+protected:
+
+	virtual void ProcessOutputSample(const TSharedRef<IMediaTextureSample, ESPMode::ThreadSafe>& Sample);
 private:
 
 	/** Mutex to ensure thread-safe access */
