Fix non-unity compile error

diff --git a/Engine/Plugins/Media/AvfMedia/Source/AvfMedia/Private/Player/AvfMediaVideoSampler.h b/Engine/Plugins/Media/AvfMedia/Source/AvfMedia/Private/Player/AvfMediaVideoSampler.h
--- a/Engine/Plugins/Media/AvfMedia/Source/AvfMedia/Private/Player/AvfMediaVideoSampler.h
+++ b/Engine/Plugins/Media/AvfMedia/Source/AvfMedia/Private/Player/AvfMediaVideoSampler.h
@@ -4,6 +4,7 @@
 
 #include "HAL/CriticalSection.h"
 #include "Templates/SharedPointer.h"
+#include "IMediaTextureSample.h"
 
 #import <AVFoundation/AVFoundation.h>
 
