REL - OpenGL needs to have the texture kept alive up until the draw happened, therefore we have to keep the refcounted resource on the stack for a bit.

diff --git a/Engine/Plugins/Media/AvfMedia/Source/AvfMedia/Private/Player/AvfMediaVideoSampler.cpp b/Engine/Plugins/Media/AvfMedia/Source/AvfMedia/Private/Player/AvfMediaVideoSampler.cpp
--- a/Engine/Plugins/Media/AvfMedia/Source/AvfMedia/Private/Player/AvfMediaVideoSampler.cpp
+++ b/Engine/Plugins/Media/AvfMedia/Source/AvfMedia/Private/Player/AvfMediaVideoSampler.cpp
@@ -288,7 +288,8 @@ void FAvfMediaVideoSampler::Tick()
 				
 				PixelShader->SetParameters(RHICmdList, YTex, UVTex, MediaShaders::YuvToSrgbPs4, true);
 				
-				RHICmdList.SetStreamSource(0, CreateTempMediaVertexBuffer(), 0);
+				FVertexBufferRHIRef VertexBuffer = CreateTempMediaVertexBuffer();
+				RHICmdList.SetStreamSource(0, VertexBuffer, 0);
 				RHICmdList.SetViewport(0, 0, 0.0f, YWidth, YHeight, 1.0f);
 
 				RHICmdList.DrawPrimitive(PT_TriangleStrip, 0, 2, 1);
diff --git a/Engine/Source/Runtime/MediaAssets/Private/Misc/MediaTextureResource.cpp b/Engine/Source/Runtime/MediaAssets/Private/Misc/MediaTextureResource.cpp
--- a/Engine/Source/Runtime/MediaAssets/Private/Misc/MediaTextureResource.cpp
+++ b/Engine/Source/Runtime/MediaAssets/Private/Misc/MediaTextureResource.cpp
@@ -633,7 +633,8 @@ void FMediaTextureResource::ConvertSample(const TSharedPtr<IMediaTextureSample,
 		}
 
 		// draw full size quad into render target
-		CommandList.SetStreamSource(0, CreateTempMediaVertexBuffer(), 0);
+		FVertexBufferRHIRef VertexBuffer = CreateTempMediaVertexBuffer();
+		CommandList.SetStreamSource(0, VertexBuffer, 0);
 		// set viewport to RT size
 		CommandList.SetViewport(0, 0, 0.0f, OutputDim.X, OutputDim.Y, 1.0f);
 
