DR - Fix deprecation warning

diff --git a/Engine/Plugins/Media/AvfMedia/Source/AvfMedia/Private/Player/AvfMediaVideoSampler.cpp b/Engine/Plugins/Media/AvfMedia/Source/AvfMedia/Private/Player/AvfMediaVideoSampler.cpp
--- a/Engine/Plugins/Media/AvfMedia/Source/AvfMedia/Private/Player/AvfMediaVideoSampler.cpp
+++ b/Engine/Plugins/Media/AvfMedia/Source/AvfMedia/Private/Player/AvfMediaVideoSampler.cpp
@@ -291,7 +291,7 @@ void FAvfMediaVideoSampler::Tick()
 				RHICmdList.SetStreamSource(0, CreateTempMediaVertexBuffer(), 0);
 				RHICmdList.SetViewport(0, 0, 0.0f, YWidth, YHeight, 1.0f);
 
-				RHICmdList.DrawPrimitive(PT_TriangleStrip, 0, 2, 1);
+				RHICmdList.DrawPrimitive(0, 2, 1);
 				
 				RHICmdList.CopyToResolveTarget(ShaderResource, ShaderResource, FResolveParams());
 			}
