MAC/iOS - Support MP4 playback from pak files using AVAssetResourceLoader delegate protocol.  Still load / play normal standalone files in the same way - only media in pak files should use the custom data loader.

diff --git a/Engine/Plugins/Media/AvfMedia/Source/AvfMedia/Private/Player/AvfMediaPlayer.cpp b/Engine/Plugins/Media/AvfMedia/Source/AvfMedia/Private/Player/AvfMediaPlayer.cpp
--- a/Engine/Plugins/Media/AvfMedia/Source/AvfMedia/Private/Player/AvfMediaPlayer.cpp
+++ b/Engine/Plugins/Media/AvfMedia/Source/AvfMedia/Private/Player/AvfMediaPlayer.cpp
@@ -14,12 +14,15 @@
 	#include "Mac/CocoaThread.h"
 #else
 	#include "IOS/IOSAsyncTask.h"
-    #include "HAL/FileManager.h"
 #endif
 
+#include "HAL/FileManager.h"
+
 #include "AvfMediaTracks.h"
 #include "AvfMediaUtils.h"
 #include "IMediaAudioSample.h"
+#include "HAL/PlatformFilemanager.h"
+#include "Async/Async.h"
 
 
 /* FAVPlayerDelegate
@@ -94,7 +97,163 @@
 
 @end
 
-/* Sync Control Class for consumed samples  */
+/* Media Resource Data loader, e.g for Pak files
+ *****************************************************************************/
+@interface FAVMediaAssetResourceLoaderDelegate : NSObject <AVAssetResourceLoaderDelegate>
+{
+	@public TSharedPtr<FArchive, ESPMode::ThreadSafe> FileAReader;
+	@public FCriticalSection CriticalSection;
+}
+@end
+
+@implementation FAVMediaAssetResourceLoaderDelegate
+
+-(FAVMediaAssetResourceLoaderDelegate*) initWithPath:(FString)InPath
+{
+	self = [super init];
+	if (self != nil)
+	{
+		Async<void>(EAsyncExecution::ThreadPool, [self, InPath]()
+		{
+			FileAReader = MakeShareable( IFileManager::Get().CreateFileReader(*InPath) );
+		});
+	}
+	return self;
+}
+
+- (void) dealloc
+{
+	{
+		FScopeLock ScopeLock(&CriticalSection);
+		
+		if(FileAReader.IsValid())
+		{
+			FileAReader->Close();
+			FileAReader = nullptr;
+		}
+	}
+	
+	[super dealloc];
+}
+
+- (BOOL)resourceLoader:(AVAssetResourceLoader *)resourceLoader shouldWaitForLoadingOfRequestedResource:(AVAssetResourceLoadingRequest *)loadingRequest
+{
+	// There should be no need to queue these up - if it turns out we need to do that - then add an ordered queue of loadingRequest objects
+	BOOL bShouldHandleLoad = NO;
+
+	FScopeLock ScopeLock(&CriticalSection);
+	
+	if(FileAReader.IsValid() && !FileAReader->IsError())
+	{
+		// Fill out content information request - if required
+		if(loadingRequest.contentInformationRequest)
+		{
+			// See https://developer.apple.com/library/archive/documentation/Miscellaneous/Reference/UTIRef/Articles/System-DeclaredUniformTypeIdentifiers.html
+			// And loadingRequest.contentInformationRequest.allowedContentTypes;
+			loadingRequest.contentInformationRequest.contentType = @"public.mpeg-4";
+			loadingRequest.contentInformationRequest.byteRangeAccessSupported = YES;
+			loadingRequest.contentInformationRequest.contentLength = FileAReader->TotalSize();
+			
+			// Handle a rare case where it only asks for content information and no data - finish it now
+			if(loadingRequest.dataRequest == nil)
+			{
+				[loadingRequest finishLoading];
+			}
+		}
+
+		// Fetch data from file - if required
+		if(loadingRequest.dataRequest)
+		{
+			// Holdon to this otherwise it'll disappear
+			[loadingRequest retain];
+			
+			// Allow this function to return so the resource loader knows the data is coming and doesn't send more requests while we block
+			Async<void>(EAsyncExecution::ThreadPool, [self, loadingRequest]()
+			{
+				bool bDataLoadSuccess = false;
+				
+				FScopeLock ScopeLock(&CriticalSection);
+				if(FileAReader.IsValid() && !FileAReader->IsError())
+				{
+					int64 Offset = loadingRequest.dataRequest.requestedOffset;
+					int64 ByteCount = loadingRequest.dataRequest.requestedLength;
+					
+					check(Offset >= 0);
+					check(ByteCount > 0);
+					
+					if(Offset + ByteCount <= FileAReader->TotalSize())
+					{
+						FileAReader->Seek(Offset);
+						
+						// Don't read the whole requested data range at once - the resource loader often asks for very large data sizes
+						// If we feed it (using respondWithData:) in chunks, it decides it has had enough data usually after a few MB,
+						// then it marks the request as cancelled, this is not an error, before issuing a different request at some point later.
+						// This keeps our peak memory usage down and limits the amount of data we are serializing.
+						
+						const int64 MaxChunkBytes = 1024 * 1024 * 1; // in single MB chunks
+						while(ByteCount > 0 && !loadingRequest.isCancelled && !FileAReader->IsError())
+						{
+							int64 ChunkByteCount = MIN(MaxChunkBytes, ByteCount);
+							ByteCount -= ChunkByteCount;
+							check(ByteCount >= 0);
+							
+							NSMutableData* nsLoadedData = [[NSMutableData alloc] initWithLength:ChunkByteCount];
+							uint8* pMemory = (uint8*)nsLoadedData.mutableBytes;
+							check(pMemory);
+						
+							FileAReader->Serialize(pMemory, ChunkByteCount);
+						
+							[loadingRequest.dataRequest respondWithData:nsLoadedData];
+							[nsLoadedData release];
+						}
+						
+						bDataLoadSuccess = !FileAReader->IsError();
+					}
+				}
+				
+				if(bDataLoadSuccess)
+				{
+					[loadingRequest finishLoading];
+				}
+				else
+				{
+					[loadingRequest finishLoadingWithError:nil];
+				}
+				
+				[loadingRequest release];
+			});
+		}
+		
+		bShouldHandleLoad = YES;
+	}
+
+	return bShouldHandleLoad;
+}
+
+- (BOOL)resourceLoader:(AVAssetResourceLoader *)resourceLoader shouldWaitForRenewalOfRequestedResource:(AVAssetResourceRenewalRequest *)renewalRequest
+{
+	// Don't set contentInformationRequest.renewalDate and we should not have to handle this case
+	return NO;
+}
+
+- (void)resourceLoader:(AVAssetResourceLoader *)resourceLoader didCancelLoadingRequest:(AVAssetResourceLoadingRequest *)loadingRequest
+{
+
+}
+
+- (BOOL)resourceLoader:(AVAssetResourceLoader *)resourceLoader shouldWaitForResponseToAuthenticationChallenge:(NSURLAuthenticationChallenge *)authenticationChallenge
+{
+	return NO;
+}
+
+- (void)resourceLoader:(AVAssetResourceLoader *)resourceLoader didCancelAuthenticationChallenge:(NSURLAuthenticationChallenge *)authenticationChallenge
+{
+
+}
+@end
+
+/* Sync Control Class for consumed samples
+ *****************************************************************************/
 class FAvfMediaSamples : public FMediaSamples
 {
 public:
@@ -158,6 +317,7 @@ FAvfMediaPlayer::FAvfMediaPlayer(IMediaEventSink& InEventSink)
 	MediaHelper = nil;
     MediaPlayer = nil;
 	PlayerItem = nil;
+	MediaResourceLoader = nil;
 		
 	bPrerolled = false;
 	bTimeSynced = false;
@@ -379,6 +539,12 @@ void FAvfMediaPlayer::Close()
 		MediaPlayer = nil;
 	}
 	
+	if(MediaResourceLoader != nil)
+	{
+		[MediaResourceLoader release];
+		MediaResourceLoader = nil;
+	}
+	
 	Tracks->Reset();
 	EventSink.ReceiveMediaEvent(EMediaEvent::TracksChanged);
 
@@ -455,7 +621,6 @@ IMediaView& FAvfMediaPlayer::GetView()
 	return *this;
 }
 
-
 bool FAvfMediaPlayer::Open(const FString& Url, const IMediaOptions* /*Options*/)
 {
 	Close();
@@ -463,12 +628,27 @@ bool FAvfMediaPlayer::Open(const FString& Url, const IMediaOptions* /*Options*/)
 	NSURL* nsMediaUrl = nil;
 	FString Path;
 
+	bool bPakResourceLoading = false;
+	
 	if (Url.StartsWith(TEXT("file://")))
 	{
 		// Media Framework doesn't percent encode the URL, so the path portion is just a native file path.
 		// Extract it and then use it create a proper URL.
 		Path = Url.Mid(7);
 		nsMediaUrl = [NSURL fileURLWithPath:Path.GetNSString() isDirectory:NO];
+		
+		// Is this from a Pak file - can't find a way to directly check - attempt to check the reverse logic
+		// as we don't want to change behaviour of normal files from a standard file URL
+		IPlatformFile& PlatformFile = FPlatformFileManager::Get().GetPlatformFile();
+		if(PlatformFile.GetLowerLevel() && !PlatformFile.GetLowerLevel()->FileExists(*Path) && FPaths::FileExists(*Path))
+		{
+			// Force the AV player to not be able to decode the scheme - this makes it use our ResourceLoader
+			NSString* formatString = [NSString stringWithFormat:@"UE4-Media://%@", [Path.GetNSString() stringByAddingPercentEncodingWithAllowedCharacters:[NSCharacterSet URLPathAllowedCharacterSet]]];
+			nsMediaUrl = [NSURL URLWithString:formatString];
+			[formatString release];
+
+			bPakResourceLoading = true;
+		}
 	}
 	else
 	{
@@ -502,7 +682,6 @@ bool FAvfMediaPlayer::Open(const FString& Url, const IMediaOptions* /*Options*/)
 	if (!MediaPlayer)
 	{
 		UE_LOG(LogAvfMedia, Error, TEXT("Failed to create instance of an AVPlayer"));
-
 		return false;
 	}
 	
@@ -512,12 +691,21 @@ bool FAvfMediaPlayer::Open(const FString& Url, const IMediaOptions* /*Options*/)
 	MediaHelper = [[FAVPlayerDelegate alloc] initWithMediaPlayer:this];
 	check(MediaHelper != nil);
 
-	PlayerItem = [[AVPlayerItem playerItemWithURL:nsMediaUrl] retain];
-
+	// Use URL asset which gives us resource loading ability if system can't handle the scheme
+	AVURLAsset* urlAsset = [[AVURLAsset alloc] initWithURL:nsMediaUrl options:nil];
+	
+	if(bPakResourceLoading)
+	{
+		MediaResourceLoader = [[FAVMediaAssetResourceLoaderDelegate alloc] initWithPath:Path];
+		[urlAsset.resourceLoader setDelegate:MediaResourceLoader queue:dispatch_get_global_queue(QOS_CLASS_DEFAULT, 0)];
+	}
+	
+	PlayerItem = [[AVPlayerItem playerItemWithAsset:urlAsset] retain];
+	[urlAsset release];
+	
 	if (PlayerItem == nil)
 	{
 		UE_LOG(LogAvfMedia, Error, TEXT("Failed to open player item with Url:"), *Url);
-
 		return false;
 	}
 
diff --git a/Engine/Plugins/Media/AvfMedia/Source/AvfMedia/Private/Player/AvfMediaPlayer.h b/Engine/Plugins/Media/AvfMedia/Source/AvfMedia/Private/Player/AvfMediaPlayer.h
--- a/Engine/Plugins/Media/AvfMedia/Source/AvfMedia/Private/Player/AvfMediaPlayer.h
+++ b/Engine/Plugins/Media/AvfMedia/Source/AvfMedia/Private/Player/AvfMediaPlayer.h
@@ -17,6 +17,7 @@ class IMediaEventSink;
 @class AVPlayer;
 @class AVPlayerItem;
 @class FAVPlayerDelegate;
+@class FAVMediaAssetResourceLoaderDelegate;
 
 
 /**
@@ -128,6 +129,9 @@ private:
 	/** Cocoa helper object we can use to keep track of ns property changes in our media items */
 	FAVPlayerDelegate* MediaHelper;
 	
+	/** Cocoa Media helper object for Pak file loading */
+	FAVMediaAssetResourceLoaderDelegate* MediaResourceLoader;
+	
 	/** The AVFoundation media player */
 	AVPlayer* MediaPlayer;
 
