Fix pak file playback race condition on initialization.  When using custom resource loader - don't preroll - causes random playback failure on iOS.

diff --git a/Engine/Plugins/Media/AvfMedia/Source/AvfMedia/Private/Player/AvfMediaPlayer.cpp b/Engine/Plugins/Media/AvfMedia/Source/AvfMedia/Private/Player/AvfMediaPlayer.cpp
--- a/Engine/Plugins/Media/AvfMedia/Source/AvfMedia/Private/Player/AvfMediaPlayer.cpp
+++ b/Engine/Plugins/Media/AvfMedia/Source/AvfMedia/Private/Player/AvfMediaPlayer.cpp
@@ -101,22 +101,22 @@
  *****************************************************************************/
 @interface FAVMediaAssetResourceLoaderDelegate : NSObject <AVAssetResourceLoaderDelegate>
 {
+	@public FString Path;
 	@public TSharedPtr<FArchive, ESPMode::ThreadSafe> FileAReader;
 	@public FCriticalSection CriticalSection;
+	@public bool bInitialized;
 }
 @end
 
 @implementation FAVMediaAssetResourceLoaderDelegate
 
--(FAVMediaAssetResourceLoaderDelegate*) initWithPath:(FString)InPath
+-(FAVMediaAssetResourceLoaderDelegate*) initWithPath:(FString&)InPath
 {
 	self = [super init];
 	if (self != nil)
 	{
-		Async<void>(EAsyncExecution::ThreadPool, [self, InPath]()
-		{
-			FileAReader = MakeShareable( IFileManager::Get().CreateFileReader(*InPath) );
-		});
+		bInitialized = false;
+		Path = InPath;
 	}
 	return self;
 }
@@ -139,95 +139,84 @@
 - (BOOL)resourceLoader:(AVAssetResourceLoader *)resourceLoader shouldWaitForLoadingOfRequestedResource:(AVAssetResourceLoadingRequest *)loadingRequest
 {
 	// There should be no need to queue these up - if it turns out we need to do that - then add an ordered queue of loadingRequest objects
-	BOOL bShouldHandleLoad = NO;
-
-	FScopeLock ScopeLock(&CriticalSection);
+	[loadingRequest retain];
 	
-	if(FileAReader.IsValid() && !FileAReader->IsError())
+	// Allow this function to return quickly so the resource loader knows the data is probabky coming and doesn't error
+	Async<void>(EAsyncExecution::ThreadPool, [self, loadingRequest]()
 	{
-		// Fill out content information request - if required
-		if(loadingRequest.contentInformationRequest)
+		FScopeLock ScopeLock(&CriticalSection);
+		
+		// If the file reader is created on the Apple callback queue then the PakLoader will throw thread errors
+		if (!bInitialized)
 		{
-			// See https://developer.apple.com/library/archive/documentation/Miscellaneous/Reference/UTIRef/Articles/System-DeclaredUniformTypeIdentifiers.html
-			// And loadingRequest.contentInformationRequest.allowedContentTypes;
-			loadingRequest.contentInformationRequest.contentType = @"public.mpeg-4";
-			loadingRequest.contentInformationRequest.byteRangeAccessSupported = YES;
-			loadingRequest.contentInformationRequest.contentLength = FileAReader->TotalSize();
-			
-			// Handle a rare case where it only asks for content information and no data - finish it now
-			if(loadingRequest.dataRequest == nil)
+			FileAReader = MakeShareable( IFileManager::Get().CreateFileReader(*Path) );
+			bInitialized = true;
+		}
+		
+		if(FileAReader.IsValid() && !FileAReader->IsError())
+		{
+			// Fill out content information request - if required
+			if(loadingRequest.contentInformationRequest)
 			{
-				[loadingRequest finishLoading];
+				// See https://developer.apple.com/library/archive/documentation/Miscellaneous/Reference/UTIRef/Articles/System-DeclaredUniformTypeIdentifiers.html
+				// And loadingRequest.contentInformationRequest.allowedContentTypes;
+				loadingRequest.contentInformationRequest.contentType = @"public.mpeg-4";
+				loadingRequest.contentInformationRequest.byteRangeAccessSupported = YES;
+				loadingRequest.contentInformationRequest.contentLength = FileAReader->TotalSize();
 			}
-		}
 
-		// Fetch data from file - if required
-		if(loadingRequest.dataRequest)
-		{
-			// Holdon to this otherwise it'll disappear
-			[loadingRequest retain];
-			
-			// Allow this function to return so the resource loader knows the data is coming and doesn't send more requests while we block
-			Async<void>(EAsyncExecution::ThreadPool, [self, loadingRequest]()
+			// Fetch data from file - if required
+			if(loadingRequest.dataRequest)
 			{
-				bool bDataLoadSuccess = false;
+				int64 Offset = loadingRequest.dataRequest.requestedOffset;
+				int64 ByteCount = loadingRequest.dataRequest.requestedLength;
+				
+				check(Offset >= 0);
+				check(ByteCount > 0);
 				
-				FScopeLock ScopeLock(&CriticalSection);
-				if(FileAReader.IsValid() && !FileAReader->IsError())
+				if(Offset + ByteCount <= FileAReader->TotalSize())
 				{
-					int64 Offset = loadingRequest.dataRequest.requestedOffset;
-					int64 ByteCount = loadingRequest.dataRequest.requestedLength;
+					FileAReader->Seek(Offset);
 					
-					check(Offset >= 0);
-					check(ByteCount > 0);
+					// Don't read the whole requested data range at once - the resource loader often asks for very large data sizes
+					// If we feed it (using respondWithData:) in chunks, it decides it has had enough data usually after a few MB,
+					// then it marks the request as cancelled, this is not an error, before issuing a different request at some point later.
+					// This keeps our peak memory usage down and limits the amount of data we are serializing.
 					
-					if(Offset + ByteCount <= FileAReader->TotalSize())
+					const int64 MaxChunkBytes = 1024 * 1024 * 1; // in single MB chunks
+					while(ByteCount > 0 && !loadingRequest.isCancelled && !FileAReader->IsError())
 					{
-						FileAReader->Seek(Offset);
+						int64 ChunkByteCount = MIN(MaxChunkBytes, ByteCount);
+						ByteCount -= ChunkByteCount;
+						check(ByteCount >= 0);
 						
-						// Don't read the whole requested data range at once - the resource loader often asks for very large data sizes
-						// If we feed it (using respondWithData:) in chunks, it decides it has had enough data usually after a few MB,
-						// then it marks the request as cancelled, this is not an error, before issuing a different request at some point later.
-						// This keeps our peak memory usage down and limits the amount of data we are serializing.
-						
-						const int64 MaxChunkBytes = 1024 * 1024 * 1; // in single MB chunks
-						while(ByteCount > 0 && !loadingRequest.isCancelled && !FileAReader->IsError())
-						{
-							int64 ChunkByteCount = MIN(MaxChunkBytes, ByteCount);
-							ByteCount -= ChunkByteCount;
-							check(ByteCount >= 0);
-							
-							NSMutableData* nsLoadedData = [[NSMutableData alloc] initWithLength:ChunkByteCount];
-							uint8* pMemory = (uint8*)nsLoadedData.mutableBytes;
-							check(pMemory);
-						
-							FileAReader->Serialize(pMemory, ChunkByteCount);
-						
-							[loadingRequest.dataRequest respondWithData:nsLoadedData];
-							[nsLoadedData release];
-						}
-						
-						bDataLoadSuccess = !FileAReader->IsError();
+						NSMutableData* nsLoadedData = [[NSMutableData alloc] initWithLength:ChunkByteCount];
+						uint8* pMemory = (uint8*)nsLoadedData.mutableBytes;
+						check(pMemory);
+					
+						FileAReader->Serialize(pMemory, ChunkByteCount);
+					
+						[loadingRequest.dataRequest respondWithData:nsLoadedData];
+						[nsLoadedData release];
 					}
 				}
-				
-				if(bDataLoadSuccess)
-				{
-					[loadingRequest finishLoading];
-				}
-				else
-				{
-					[loadingRequest finishLoadingWithError:nil];
-				}
-				
-				[loadingRequest release];
-			});
+			}
 		}
 		
-		bShouldHandleLoad = YES;
-	}
+		// Check file reader is not in error state after potential seek and data read operations
+		if(FileAReader.IsValid() && !FileAReader->IsError())
+		{
+			[loadingRequest finishLoading];
+		}
+		else
+		{
+			[loadingRequest finishLoadingWithError:nil];
+		}
+
+		[loadingRequest release];
+	});
 
-	return bShouldHandleLoad;
+	return YES;
 }
 
 - (BOOL)resourceLoader:(AVAssetResourceLoader *)resourceLoader shouldWaitForRenewalOfRequestedResource:(AVAssetResourceRenewalRequest *)renewalRequest
@@ -414,39 +403,40 @@ void FAvfMediaPlayer::OnStatusNotification()
 
 					if (!bPrerolled)
 					{
-						// Preroll for playback.
-						[MediaPlayer prerollAtRate:1.0f completionHandler:^(BOOL bFinished)
+						if(MediaResourceLoader != nil)
+						{
+							// If there is a resource loader - don't preroll
+							bPrerolled = true;
+							CurrentState = EMediaState::Stopped;
+							EventSink.ReceiveMediaEvent(EMediaEvent::MediaOpened);
+						}
+						else
 						{
-							if (bFinished)
+							// Preroll for playback.
+							[MediaPlayer prerollAtRate:1.0f completionHandler:^(BOOL bFinished)
 							{
-								PlayerTasks.Enqueue([=]()
+								if (bFinished)
 								{
-									if(PlayerItem.status == AVPlayerItemStatusReadyToPlay)
+									PlayerTasks.Enqueue([=]()
 									{
-										bPrerolled = true;
-										CurrentState = EMediaState::Stopped;
-										
-										EventSink.ReceiveMediaEvent(EMediaEvent::MediaOpened);
-										if (CurrentTime != FTimespan::Zero())
+										if(PlayerItem.status == AVPlayerItemStatusReadyToPlay)
 										{
-											Seek(CurrentTime);
+											bPrerolled = true;
+											CurrentState = EMediaState::Stopped;
+											EventSink.ReceiveMediaEvent(EMediaEvent::MediaOpened);
 										}
-										if(CurrentRate != 0.0f)
-										{
-											SetRate(CurrentRate);
-										}
-									}
-								});
-							}
-							else
-							{
-								PlayerTasks.Enqueue([=]()
+									});
+								}
+								else
 								{
-									CurrentState = EMediaState::Error;
-									EventSink.ReceiveMediaEvent(EMediaEvent::MediaOpenFailed);
-								});
-							}
-						}];
+									PlayerTasks.Enqueue([=]()
+									{
+										CurrentState = EMediaState::Error;
+										EventSink.ReceiveMediaEvent(EMediaEvent::MediaOpenFailed);
+									});
+								}
+							}];
+						}
 					}
 				}
 
