Use the subtype for encoding checks not the extension FormatType.

diff --git a/Engine/Plugins/Media/AvfMedia/Source/AvfMedia/Private/Player/AvfMediaTracks.cpp b/Engine/Plugins/Media/AvfMedia/Source/AvfMedia/Private/Player/AvfMediaTracks.cpp
--- a/Engine/Plugins/Media/AvfMedia/Source/AvfMedia/Private/Player/AvfMediaTracks.cpp
+++ b/Engine/Plugins/Media/AvfMedia/Source/AvfMedia/Private/Player/AvfMediaTracks.cpp
@@ -448,38 +448,28 @@ void FAvfMediaTracks::Initialize(AVPlayerItem* InPlayerItem, FString& OutInfo)
 			
 			NSMutableDictionary* OutputSettings = [NSMutableDictionary dictionary];
 			CMFormatDescriptionRef DescRef = (CMFormatDescriptionRef)[AssetTrack.formatDescriptions objectAtIndex:0];
+			CMVideoCodecType CodecType = CMFormatDescriptionGetMediaSubType(DescRef);
 			
-			// Select decode format
-			int32 DecodePixelBufferFormat = 0;
+			// Select decode pixel format - BGRA32 is the fallback - Any more pixel formats added here need to be handled correctly in FAvfMediaVideoSampler
+			int32 DecodePixelBufferFormat = kCVPixelFormatType_32BGRA;
 	
 			if(!AVFMediaForceDecodeBGRA)
 			{
-				CFDictionaryRef FormatExtensions = CMFormatDescriptionGetExtensions(DescRef);
-				if(FormatExtensions)
+				if(kCMVideoCodecType_H264 == CodecType)
 				{
-					CFStringRef FormatName = (CFStringRef)CFDictionaryGetValue(FormatExtensions, kCMFormatDescriptionExtension_FormatName);
+					DecodePixelBufferFormat = kCVPixelFormatType_420YpCbCr8BiPlanarVideoRange;
 					
-					// kCMVideoCodecType_H264 : 'avc1' - why is this a CFString in the dictionary and not a CFNumber?
-					if(FormatName && kCFCompareEqualTo == CFStringCompare(FormatName, CFSTR("'avc1'"), 0))
+					CFDictionaryRef FormatExtensions = CMFormatDescriptionGetExtensions(DescRef);
+					if(FormatExtensions)
 					{
 						CFBooleanRef bFullRange = (CFBooleanRef)CFDictionaryGetValue(FormatExtensions, kCMFormatDescriptionExtension_FullRangeVideo);
 						if(bFullRange && (bool)CFBooleanGetValue(bFullRange))
 						{
 							DecodePixelBufferFormat = kCVPixelFormatType_420YpCbCr8BiPlanarFullRange;
 						}
-						else
-						{
-							DecodePixelBufferFormat = kCVPixelFormatType_420YpCbCr8BiPlanarVideoRange;
-						}
 					}
 				}
 			}
-
-			// BGRA32 is the fallback
-			if(DecodePixelBufferFormat == 0)
-			{
-				DecodePixelBufferFormat = kCVPixelFormatType_32BGRA;
-			}
 			
 			[OutputSettings setObject : [NSNumber numberWithInt :  DecodePixelBufferFormat] forKey : (NSString*)kCVPixelBufferPixelFormatTypeKey];
 			
@@ -524,7 +514,6 @@ void FAvfMediaTracks::Initialize(AVPlayerItem* InPlayerItem, FString& OutInfo)
 			CGSize NaturalFrameSize = AssetTrack.naturalSize;
 			Track->FrameSize = FIntPoint((int32)NaturalFrameSize.width, (int32)NaturalFrameSize.height);
 
-			CMVideoCodecType CodecType = CMFormatDescriptionGetMediaSubType(DescRef);
 			OutInfo += FString::Printf(TEXT("    Codec: %s\n"), *AvfMedia::CodecTypeToString(CodecType));
 			OutInfo += FString::Printf(TEXT("    Dimensions: %i x %i\n"), (int32)NaturalFrameSize.width, (int32)NaturalFrameSize.height);
 			OutInfo += FString::Printf(TEXT("    Frame Rate: %g fps\n"), AssetTrack.nominalFrameRate);
