Remove video sample queue time sync and use default video player time when there is no audio track.  The time sync's were put in to allow the high level to know when audio had been decoded.  But without an audio track this higher level system breaks down and we just don't move on - no need to try an please that system when it's video only - just use player time for the sync point.  Also remove deprecated function - would have no effect and any seeking should be done thorugh the ::seek() function.

diff --git a/Engine/Plugins/Media/AvfMedia/Source/AvfMedia/Private/Player/AvfMediaPlayer.cpp b/Engine/Plugins/Media/AvfMedia/Source/AvfMedia/Private/Player/AvfMediaPlayer.cpp
--- a/Engine/Plugins/Media/AvfMedia/Source/AvfMedia/Private/Player/AvfMediaPlayer.cpp
+++ b/Engine/Plugins/Media/AvfMedia/Source/AvfMedia/Private/Player/AvfMediaPlayer.cpp
@@ -740,15 +740,12 @@ bool FAvfMediaPlayer::Open(const FString& Url, const IMediaOptions* /*Options*/)
 
 	[[NSNotificationCenter defaultCenter] addObserver:MediaHelper selector:@selector(playerItemPlaybackEndReached:) name:AVPlayerItemDidPlayToEndTimeNotification object:PlayerItem];
 	[PlayerItem addObserver:MediaHelper forKeyPath:@"status" options:0 context:PlayerItem];
-
-	[MediaPlayer replaceCurrentItemWithPlayerItem : PlayerItem];
-	PRAGMA_DISABLE_DEPRECATION_WARNINGS
-	[[MediaPlayer currentItem] seekToTime:kCMTimeZero];
-	PRAGMA_ENABLE_DEPRECATION_WARNINGS
-
+	
 	MediaPlayer.rate = 0.0;
 	CurrentTime = FTimespan::Zero();
 
+	[MediaPlayer replaceCurrentItemWithPlayerItem : PlayerItem];
+
 	if (!EnteredForegroundHandle.IsValid())
     {
         EnteredForegroundHandle = FCoreDelegates::ApplicationHasEnteredForegroundDelegate.AddRaw(this, &FAvfMediaPlayer::HandleApplicationHasEnteredForeground);
@@ -819,10 +816,6 @@ void FAvfMediaPlayer::TickInput(FTimespan DeltaTime, FTimespan /*Timecode*/)
 						{
 							SyncTime = GetAudioTimeSync();
 						}
-						else if(Tracks->GetSelectedTrack(EMediaTrackType::Video) != INDEX_NONE)
-						{
-							SyncTime = GetVideoTimeSync();
-						}
 						else /* Default Use AVPlayer time*/
 #endif
 						{
