Fix for looping while playing in reverse when playing audio through audio sinks - can't use audio sample sync when in reverse, we don't process audio samples with -ve rates.

diff --git a/Engine/Plugins/Media/AvfMedia/Source/AvfMedia/Private/Player/AvfMediaPlayer.cpp b/Engine/Plugins/Media/AvfMedia/Source/AvfMedia/Private/Player/AvfMediaPlayer.cpp
--- a/Engine/Plugins/Media/AvfMedia/Source/AvfMedia/Private/Player/AvfMediaPlayer.cpp
+++ b/Engine/Plugins/Media/AvfMedia/Source/AvfMedia/Private/Player/AvfMediaPlayer.cpp
@@ -812,7 +812,8 @@ void FAvfMediaPlayer::TickInput(FTimespan DeltaTime, FTimespan /*Timecode*/)
 					{
 						FTimespan SyncTime = FTimespan::MinValue();
 #if AUDIO_PLAYBACK_VIA_ENGINE
-						if(Tracks->GetSelectedTrack(EMediaTrackType::Audio) != INDEX_NONE)
+						// There is no audio in reverse - can't use as sync point - same issue as if no audio track video
+						if(Tracks->GetSelectedTrack(EMediaTrackType::Audio) != INDEX_NONE && CurrentRate >= 0.f)
 						{
 							SyncTime = GetAudioTimeSync();
 						}
