merge in the video player memory fix for IOS

diff --git a/Engine/Plugins/Media/AvfMedia/Source/AvfMedia/Private/Player/AvfMediaPlayer.cpp b/Engine/Plugins/Media/AvfMedia/Source/AvfMedia/Private/Player/AvfMediaPlayer.cpp
--- a/Engine/Plugins/Media/AvfMedia/Source/AvfMedia/Private/Player/AvfMediaPlayer.cpp
+++ b/Engine/Plugins/Media/AvfMedia/Source/AvfMedia/Private/Player/AvfMediaPlayer.cpp
@@ -321,6 +321,9 @@ FAvfMediaPlayer::~FAvfMediaPlayer()
 {
 	Close();
 
+	delete Tracks;
+	Tracks = nullptr;
+	
 	delete Samples;
 	Samples = nullptr;
 }
diff --git a/Engine/Plugins/Media/AvfMedia/Source/AvfMedia/Private/Player/AvfMediaTextureSample.h b/Engine/Plugins/Media/AvfMedia/Source/AvfMedia/Private/Player/AvfMediaTextureSample.h
--- a/Engine/Plugins/Media/AvfMedia/Source/AvfMedia/Private/Player/AvfMediaTextureSample.h
+++ b/Engine/Plugins/Media/AvfMedia/Source/AvfMedia/Private/Player/AvfMediaTextureSample.h
@@ -163,6 +163,11 @@ public:
 		return true;
 	}
 
+	virtual void Reset() override
+	{
+		Texture = nullptr;
+	}
+	
 private:
 
 	/** Width and height of the texture sample. */
diff --git a/Engine/Plugins/Media/AvfMedia/Source/AvfMedia/Private/Player/AvfMediaTracks.cpp b/Engine/Plugins/Media/AvfMedia/Source/AvfMedia/Private/Player/AvfMediaTracks.cpp
--- a/Engine/Plugins/Media/AvfMedia/Source/AvfMedia/Private/Player/AvfMediaTracks.cpp
+++ b/Engine/Plugins/Media/AvfMedia/Source/AvfMedia/Private/Player/AvfMediaTracks.cpp
@@ -627,6 +627,7 @@ void FAvfMediaTracks::Reset()
 			if (PinnedVideoSampler.IsValid())
 			{
 				PinnedVideoSampler->SetOutput(nil, 0.0f);
+				PinnedVideoSampler->Reset();
 			}
 		});
 	
diff --git a/Engine/Plugins/Media/AvfMedia/Source/AvfMedia/Private/Player/AvfMediaVideoSampler.cpp b/Engine/Plugins/Media/AvfMedia/Source/AvfMedia/Private/Player/AvfMediaVideoSampler.cpp
--- a/Engine/Plugins/Media/AvfMedia/Source/AvfMedia/Private/Player/AvfMediaVideoSampler.cpp
+++ b/Engine/Plugins/Media/AvfMedia/Source/AvfMedia/Private/Player/AvfMediaVideoSampler.cpp
@@ -181,6 +181,11 @@ void FAvfMediaVideoSampler::SetOutput(AVPlayerItemVideoOutput* InOutput, float I
 	}
 }
 
+void FAvfMediaVideoSampler::Reset()
+{
+	auto VideoSample = VideoSamplePool->AcquireShared();
+	VideoSample->Reset();
+}
 
 void FAvfMediaVideoSampler::Tick()
 {
diff --git a/Engine/Plugins/Media/AvfMedia/Source/AvfMedia/Private/Player/AvfMediaVideoSampler.h b/Engine/Plugins/Media/AvfMedia/Source/AvfMedia/Private/Player/AvfMediaVideoSampler.h
--- a/Engine/Plugins/Media/AvfMedia/Source/AvfMedia/Private/Player/AvfMediaVideoSampler.h
+++ b/Engine/Plugins/Media/AvfMedia/Source/AvfMedia/Private/Player/AvfMediaVideoSampler.h
@@ -4,6 +4,7 @@
 
 #include "HAL/CriticalSection.h"
 #include "Templates/SharedPointer.h"
+
 #include "IMediaTextureSample.h"
 
 #import <AVFoundation/AVFoundation.h>
@@ -41,6 +42,8 @@ public:
 
 	void ProcessFrame(CVPixelBufferRef Frame, FTimespan SampleTime, FTimespan SampleDuration);
 
+	void Reset();
+	
 protected:
 
 	virtual void ProcessOutputSample(const TSharedRef<IMediaTextureSample, ESPMode::ThreadSafe>& Sample);
diff --git a/Engine/Source/Runtime/Media/Public/IMediaTextureSample.h b/Engine/Source/Runtime/Media/Public/IMediaTextureSample.h
--- a/Engine/Source/Runtime/Media/Public/IMediaTextureSample.h
+++ b/Engine/Source/Runtime/Media/Public/IMediaTextureSample.h
@@ -249,6 +249,8 @@ public:
 		return DefaultMatrix;
 	}
 
+	virtual void Reset() { }
+	
 public:
 
 	/** Virtual destructor. */
