Apple build fix

diff --git a/Engine/Plugins/Media/AvfMedia/Source/AvfMedia/Private/Player/AvfMediaVideoSampler.cpp b/Engine/Plugins/Media/AvfMedia/Source/AvfMedia/Private/Player/AvfMediaVideoSampler.cpp
--- a/Engine/Plugins/Media/AvfMedia/Source/AvfMedia/Private/Player/AvfMediaVideoSampler.cpp
+++ b/Engine/Plugins/Media/AvfMedia/Source/AvfMedia/Private/Player/AvfMediaVideoSampler.cpp
@@ -272,18 +272,18 @@ void FAvfMediaVideoSampler::ProcessFrame(CVPixelBufferRef Frame, FTimespan Sampl
 			check(UVTextureRef);
 			
 			// Metal can upload directly from an IOSurface to a 2D texture, so we can just wrap it.
-			FRHIResourceCreateInfo YCreateInfo;
+			FRHIResourceCreateInfo YCreateInfo(TEXT("YTex"));
 			YCreateInfo.BulkData = new FAvfTexture2DResourceWrapper(YTextureRef);
 			YCreateInfo.ResourceArray = nullptr;
 			
-			FRHIResourceCreateInfo UVCreateInfo;
+			FRHIResourceCreateInfo UVCreateInfo(TEXT("UVTex"));
 			UVCreateInfo.BulkData = new FAvfTexture2DResourceWrapper(UVTextureRef);
 			UVCreateInfo.ResourceArray = nullptr;
 			
 			TRefCountPtr<FRHITexture2D> YTex = RHICreateTexture2D(YWidth, YHeight, PF_G8, 1, 1, TexCreateFlags | TexCreate_ShaderResource, YCreateInfo);
 			TRefCountPtr<FRHITexture2D> UVTex = RHICreateTexture2D(UVWidth, UVHeight, PF_R8G8, 1, 1, TexCreateFlags | TexCreate_ShaderResource, UVCreateInfo);
 			
-			FRHIResourceCreateInfo Info;
+			FRHIResourceCreateInfo Info(TEXT("Info"));
 			ShaderResource = RHICreateTexture2D(YWidth, YHeight, PF_B8G8R8A8, 1, 1, TexCreateFlags | TexCreate_ShaderResource | TexCreate_RenderTargetable | TexCreate_SRGB, Info);
 			
 			// render video frame into sink texture
@@ -336,7 +336,7 @@ void FAvfMediaVideoSampler::ProcessFrame(CVPixelBufferRef Frame, FTimespan Sampl
 			check(Result == kCVReturnSuccess);
 			check(TextureRef);
 			
-			FRHIResourceCreateInfo CreateInfo;
+			FRHIResourceCreateInfo CreateInfo(TEXT("FAvfMediaVideoSampler"));
 			CreateInfo.BulkData = new FAvfTexture2DResourceWrapper(TextureRef);
 			CreateInfo.ResourceArray = nullptr;
 			
diff --git a/Engine/Plugins/Media/ElectraUtil/Source/ElectraSamples/Private/ElectraTextureSampleAVF.cpp b/Engine/Plugins/Media/ElectraUtil/Source/ElectraSamples/Private/ElectraTextureSampleAVF.cpp
--- a/Engine/Plugins/Media/ElectraUtil/Source/ElectraSamples/Private/ElectraTextureSampleAVF.cpp
+++ b/Engine/Plugins/Media/ElectraUtil/Source/ElectraSamples/Private/ElectraTextureSampleAVF.cpp
@@ -209,11 +209,11 @@ void FElectraMediaTexConvApple::ConvertTexture(FTexture2DRHIRef & InDstTexture,
 			check(UVTextureRef);
 
 			// Metal can upload directly from an IOSurface to a 2D texture, so we can just wrap it.
-			FRHIResourceCreateInfo YCreateInfo;
+			FRHIResourceCreateInfo YCreateInfo(TEXT("YTex"));
 			YCreateInfo.BulkData = new FTexConvTexResourceWrapper(YTextureRef);
 			YCreateInfo.ResourceArray = nullptr;
 
-			FRHIResourceCreateInfo UVCreateInfo;
+			FRHIResourceCreateInfo UVCreateInfo(TEXT("UVTex"));
 			UVCreateInfo.BulkData = new FTexConvTexResourceWrapper(UVTextureRef);
 			UVCreateInfo.ResourceArray = nullptr;
 
@@ -274,7 +274,7 @@ void FElectraMediaTexConvApple::ConvertTexture(FTexture2DRHIRef & InDstTexture,
 			check(Result == kCVReturnSuccess);
 			check(TextureRef);
 
-			FRHIResourceCreateInfo CreateInfo;
+			FRHIResourceCreateInfo CreateInfo(TEXT("DstTexture"));
 			CreateInfo.BulkData = new FTexConvTexResourceWrapper(TextureRef);
 			CreateInfo.ResourceArray = nullptr;
 
diff --git a/Engine/Plugins/Runtime/AR/AppleAR/AppleARKit/Source/AppleARKit/Private/AppleARKitTextures.cpp b/Engine/Plugins/Runtime/AR/AppleAR/AppleARKit/Source/AppleARKit/Private/AppleARKitTextures.cpp
--- a/Engine/Plugins/Runtime/AR/AppleAR/AppleARKit/Source/AppleARKit/Private/AppleARKitTextures.cpp
+++ b/Engine/Plugins/Runtime/AR/AppleAR/AppleARKit/Source/AppleARKit/Private/AppleARKitTextures.cpp
@@ -207,7 +207,7 @@ protected:
 			TextureRHI.SafeRelease();
 			
 			// Create the target texture that we'll update into
-			FRHIResourceCreateInfo CreateInfo;
+			FRHIResourceCreateInfo CreateInfo(TEXT("FARKitTextureResource"));
 			TextureRHI = RHICreateTexture2D(Size.X, Size.Y, PixelFormat, 1, 1, TexCreate_Dynamic | TexCreate_ShaderResource | TexCreate_UAV, CreateInfo);
 		}
 	}
@@ -571,7 +571,7 @@ public:
 	 */
 	virtual void InitRHI() override
 	{
-		FRHIResourceCreateInfo CreateInfo;
+		FRHIResourceCreateInfo CreateInfo(TEXT("FARMetalResource"));
 
 		id<MTLTexture> MetalTexture = Owner->GetMetalTexture();
 		if (MetalTexture != nullptr)
@@ -940,7 +940,7 @@ public:
 		
 		{
 			// Update the RHI texture wrapper for the Y and CbCr images
-			FRHIResourceCreateInfo CreateInfo;
+			FRHIResourceCreateInfo CreateInfo(TEXT("VideoTexture"));
 			const ETextureCreateFlags CreateFlags = TexCreate_Dynamic | TexCreate_NoTiling | TexCreate_ShaderResource;
 			CreateInfo.BulkData = new FAppleARKitCameraTextureResourceWrapper(CapturedYImage);
 			CreateInfo.ResourceArray = nullptr;
diff --git a/Engine/Plugins/Runtime/AppleMoviePlayer/Source/AppleMoviePlayer/Private/AppleMovieStreamer.cpp b/Engine/Plugins/Runtime/AppleMoviePlayer/Source/AppleMoviePlayer/Private/AppleMovieStreamer.cpp
--- a/Engine/Plugins/Runtime/AppleMoviePlayer/Source/AppleMoviePlayer/Private/AppleMovieStreamer.cpp
+++ b/Engine/Plugins/Runtime/AppleMoviePlayer/Source/AppleMoviePlayer/Private/AppleMovieStreamer.cpp
@@ -586,7 +586,7 @@ bool FAVPlayerMovieStreamer::CheckForNextFrameAndCopy()
 		check(Result == kCVReturnSuccess);
 		check(TextureCacheRef);
 	
-		FRHIResourceCreateInfo CreateInfo(new FAppleMovieStreamerTexture2DResourceWrapper(TextureCacheRef));
+		FRHIResourceCreateInfo CreateInfo(TEXT("FAVPlayerMovieStreamer"), new FAppleMovieStreamerTexture2DResourceWrapper(TextureCacheRef));
 		FTexture2DRHIRef MediaWrappedTexture = RHICreateTexture2D(Width, Height, PF_B8G8R8A8, 1, 1, TexCreate_Dynamic | TexCreate_NoTiling, CreateInfo);
 		
 		if(SlateTexture.IsValid())
