AvfMedia: Fixed UE4 reference.
#rnx

diff --git a/Engine/Plugins/Media/AvfMedia/Source/AvfMedia/Private/Player/AvfMediaPlayer.cpp b/Engine/Plugins/Media/AvfMedia/Source/AvfMedia/Private/Player/AvfMediaPlayer.cpp
--- a/Engine/Plugins/Media/AvfMedia/Source/AvfMedia/Private/Player/AvfMediaPlayer.cpp
+++ b/Engine/Plugins/Media/AvfMedia/Source/AvfMedia/Private/Player/AvfMediaPlayer.cpp
@@ -644,7 +644,7 @@ bool FAvfMediaPlayer::Open(const FString& Url, const IMediaOptions* /*Options*/)
 		if (PakFileMgr != nullptr && PakFileMgr->FindFileInPakFiles(*Path, nullptr, nullptr))
 		{
 			// Force the AV player to not be able to decode the scheme - this makes it use our ResourceLoader
-			NSString* formatString = [NSString stringWithFormat:@"UE4-Media://%@", [Path.GetNSString() stringByAddingPercentEncodingWithAllowedCharacters:[NSCharacterSet URLPathAllowedCharacterSet]]];
+			NSString* formatString = [NSString stringWithFormat:@"UE-Media://%@", [Path.GetNSString() stringByAddingPercentEncodingWithAllowedCharacters:[NSCharacterSet URLPathAllowedCharacterSet]]];
 			nsMediaUrl = [NSURL URLWithString:formatString];
 
 			bPakResourceLoading = true;
