AvfMedia: Fixed UE4 reference.
#rnx

diff --git a/Engine/Plugins/Media/AvfMedia/Source/AvfMedia/Private/Player/AvfMediaTracks.cpp b/Engine/Plugins/Media/AvfMedia/Source/AvfMedia/Private/Player/AvfMediaTracks.cpp
--- a/Engine/Plugins/Media/AvfMedia/Source/AvfMedia/Private/Player/AvfMediaTracks.cpp
+++ b/Engine/Plugins/Media/AvfMedia/Source/AvfMedia/Private/Player/AvfMediaTracks.cpp
@@ -24,8 +24,8 @@ static int32 AVFMediaForceDecodeBGRA = 0;
 FAutoConsoleVariableRef CVarAVFMediaForceDecodeBGRA(
 	TEXT("m.avf.ForceDecodeBGRA"),
 	AVFMediaForceDecodeBGRA,
-	TEXT("Change between YUV decode and convert to BGRA in UE4 Shader (Keeps everything on the GPU) or always force Apple framework to perform the decode to BRGA (potential performance penalty).\n")
-	TEXT("0: Auto Detect YUV (Default) and decode to BGRA in UE4 Shader, 1: Force AV Framework to decode to BGRA"),
+	TEXT("Change between YUV decode and convert to BGRA in UE Shader (Keeps everything on the GPU) or always force Apple framework to perform the decode to BRGA (potential performance penalty).\n")
+	TEXT("0: Auto Detect YUV (Default) and decode to BGRA in UE Shader, 1: Force AV Framework to decode to BGRA"),
 	ECVF_ReadOnly);
 
 /* FAVPlayerItemLegibleOutputPushDelegate
