Mac/iOS IMediaCaptureSupport implementation for Mic and Webcam.

diff --git a/Engine/Plugins/Media/AvfMedia/AvfMedia.uplugin b/Engine/Plugins/Media/AvfMedia/AvfMedia.uplugin
--- a/Engine/Plugins/Media/AvfMedia/AvfMedia.uplugin
+++ b/Engine/Plugins/Media/AvfMedia/AvfMedia.uplugin
@@ -37,6 +37,12 @@
 			"Type" : "RuntimeNoCommandlet",
 			"LoadingPhase" : "PostEngineInit",
 			"PlatformAllowList" : [ "IOS", "Linux", "Mac", "TVOS", "Win64" ]
+		},
+		{
+			"Name" : "AvfMediaCapture",
+			"Type" : "RuntimeNoCommandlet",
+			"LoadingPhase" : "Default",
+			"PlatformAllowList" : [ "IOS", "Mac" ]
 		}
 	]
-}
\ No newline at end of file
+}
diff --git a/Engine/Plugins/Media/AvfMedia/Source/AvfMedia/Private/Player/AvfMediaAudioSample.h b/Engine/Plugins/Media/AvfMedia/Source/AvfMedia/Private/Player/AvfMediaAudioSample.h
deleted file mode 100644
diff --git a/Engine/Plugins/Media/AvfMedia/Source/AvfMedia/Private/Player/AvfMediaTextureSample.h b/Engine/Plugins/Media/AvfMedia/Source/AvfMedia/Private/Player/AvfMediaTextureSample.h
deleted file mode 100644
diff --git a/Engine/Plugins/Media/AvfMedia/Source/AvfMedia/Private/Player/AvfMediaVideoSampler.cpp b/Engine/Plugins/Media/AvfMedia/Source/AvfMedia/Private/Player/AvfMediaVideoSampler.cpp
--- a/Engine/Plugins/Media/AvfMedia/Source/AvfMedia/Private/Player/AvfMediaVideoSampler.cpp
+++ b/Engine/Plugins/Media/AvfMedia/Source/AvfMedia/Private/Player/AvfMediaVideoSampler.cpp
@@ -21,54 +21,6 @@
 #include "AvfMediaTextureSample.h"
 
 
-/**
- * Passes a CV*TextureRef or CVPixelBufferRef through to the RHI to wrap in an RHI texture without traversing system memory.
- */
-class FAvfTexture2DResourceWrapper
-	: public FResourceBulkDataInterface
-{
-public:
-
-	FAvfTexture2DResourceWrapper(CFTypeRef InImageBuffer)
-		: ImageBuffer(InImageBuffer)
-	{
-		check(ImageBuffer);
-		CFRetain(ImageBuffer);
-	}
-
-	virtual ~FAvfTexture2DResourceWrapper()
-	{
-		CFRelease(ImageBuffer);
-		ImageBuffer = nullptr;
-	}
-
-public:
-
-	//~ FResourceBulkDataInterface interface
-
-	virtual void Discard() override
-	{
-		delete this;
-	}
-
-	virtual const void* GetResourceBulkData() const override
-	{
-		return ImageBuffer;
-	}
-	
-	virtual uint32 GetResourceBulkDataSize() const override
-	{
-		return ImageBuffer ? ~0u : 0;
-	}
-
-	virtual EBulkDataType GetResourceType() const override
-	{
-		return EBulkDataType::MediaTexture;
-	}
-			
-	CFTypeRef ImageBuffer;
-};
-
 
 /**
  * Allows for direct GPU mem allocation for texture resource from a CVImageBufferRef's system memory backing store.
diff --git a/Engine/Plugins/Media/AvfMedia/Source/AvfMedia/Public/AvfMediaAudioSample.h b/Engine/Plugins/Media/AvfMedia/Source/AvfMedia/Public/AvfMediaAudioSample.h
new file mode 100644
--- /dev/null
+++ b/Engine/Plugins/Media/AvfMedia/Source/AvfMedia/Public/AvfMediaAudioSample.h
@@ -0,0 +1,148 @@
+// Copyright Epic Games, Inc. All Rights Reserved.
+
+#pragma once
+
+#include "CoreTypes.h"
+#include "Containers/Array.h"
+#include "IMediaAudioSample.h"
+#include "MediaObjectPool.h"
+#include "Misc/Timespan.h"
+
+
+/**
+ * Audio sample generated by AvfMedia player.
+ */
+class FAvfMediaAudioSample
+	: public IMediaAudioSample
+	, public IMediaPoolable
+{
+public:
+
+	/** Default constructor. */
+	FAvfMediaAudioSample()
+		: Channels(0)
+		, Duration(0)
+		, Frames(0)
+		, SampleRate(0)
+		, Time(FTimespan::Zero())
+		, Format(EMediaAudioSampleFormat::Undefined)
+	{ }
+
+	/** Virtual destructor. */
+	virtual ~FAvfMediaAudioSample() { }
+
+public:
+
+	/**
+	 * Get a writable pointer to the sample buffer.
+	 *
+	 * @return Sample buffer.
+	 */
+	uint8* GetMutableBuffer()
+	{
+		return Buffer.GetData();
+	}
+
+	/**
+	 * Initialize the sample.
+	 *
+	 * @param InBufferSize The size of the audio buffer.
+	 * @param InFrames Number of frames in the buffer.
+	 * @param InChannels Number of audio channels.
+	 * @param InSampleRate The sample rate.
+	 * @param InTime The sample time (in the player's local clock).
+	 * @return true on success, false otherwise.
+	 */
+	bool Initialize(
+		uint32 InBufferSize,
+		uint32 InFrames,
+		uint32 InChannels,
+		uint32 InSampleRate,
+		FTimespan InTime,
+		FTimespan InDuration,
+		EMediaAudioSampleFormat InFormat = EMediaAudioSampleFormat::Float)
+	{
+		if (InBufferSize == 0)
+		{
+			return false;
+		}
+
+		Buffer.Reset(InBufferSize);
+		Buffer.AddUninitialized(InBufferSize);
+
+		Channels = InChannels;
+		Frames = InFrames;
+		SampleRate = InSampleRate;
+		Time = InTime;
+		Duration = InDuration;
+		Format = InFormat;
+
+		return true;
+	}
+
+public:
+
+	//~ IMediaAudioSample interface
+
+	virtual const void* GetBuffer() override
+	{
+		return Buffer.GetData();
+	}
+
+	virtual uint32 GetChannels() const override
+	{
+		return Channels;
+	}
+
+	virtual FTimespan GetDuration() const override
+	{
+		return Duration;
+	}
+
+	virtual EMediaAudioSampleFormat GetFormat() const override
+	{
+		return Format;
+	}
+
+	virtual uint32 GetFrames() const override
+	{
+		return Frames;
+	}
+
+	virtual uint32 GetSampleRate() const override
+	{
+		return SampleRate;
+	}
+
+	virtual FMediaTimeStamp GetTime() const override
+	{
+		return FMediaTimeStamp(Time);
+	}
+
+private:
+
+	/** The sample's frame buffer. */
+	TArray<uint8> Buffer;
+
+	/** Number of audio channels. */
+	uint32 Channels;
+
+	/** Duration for which the sample is valid. */
+	FTimespan Duration;
+
+	/** Number of frames in the buffer. */
+	uint32 Frames;
+
+	/** Audio sample rate (in samples per second). */
+	uint32 SampleRate;
+
+	/** Play time for which the sample was generated. */
+	FTimespan Time;
+	
+	/** Sample Format  (8/16/24/32 Integer or 32 Float) */
+	EMediaAudioSampleFormat Format;
+};
+
+
+/** Implements a pool for AVF audio sample objects. */
+class FAvfMediaAudioSamplePool : public TMediaObjectPool<FAvfMediaAudioSample> { };
diff --git a/Engine/Plugins/Media/AvfMedia/Source/AvfMedia/Public/AvfMediaTextureSample.h b/Engine/Plugins/Media/AvfMedia/Source/AvfMedia/Public/AvfMediaTextureSample.h
new file mode 100644
--- /dev/null
+++ b/Engine/Plugins/Media/AvfMedia/Source/AvfMedia/Public/AvfMediaTextureSample.h
@@ -0,0 +1,249 @@
+// Copyright Epic Games, Inc. All Rights Reserved.
+
+#pragma once
+
+#include "CoreTypes.h"
+#include "IMediaTextureSample.h"
+#include "MediaObjectPool.h"
+#include "Math/IntPoint.h"
+#include "Math/Range.h"
+#include "Misc/Timespan.h"
+
+#if WITH_ENGINE
+	#include "RHI.h"
+#else
+	#include "Containers/Array.h"
+#endif
+
+
+/**
+ * Texture sample generated by AvfMedia player.
+ */
+class FAvfMediaTextureSample
+	: public IMediaTextureSample
+	, public IMediaPoolable
+{
+public:
+
+	/** Virtual destructor. */
+	virtual ~FAvfMediaTextureSample() { }
+
+public:
+
+#if WITH_ENGINE
+
+	/**
+	 * Initialize the texture sample with a texture resource.
+	 *
+	 * @param InTexture The texture resource containing the video frame.
+	 * @param InDim The sample buffer's width and height (in pixels).
+	 * @param InOutputDim The sample's output width and height (in pixels).
+	 * @param InTime The sample time (in the player's local clock).
+	 * @param InDuration The duration for which the sample is valid.
+	 * @return true on success, false otherwise.
+	 */
+	bool Initialize(
+		TRefCountPtr<FRHITexture2D>& InTexture,
+		const FIntPoint& InDim,
+		const FIntPoint& InOutputDim,
+		FTimespan InTime,
+		FTimespan InDuration)
+	{
+		if (!InTexture.IsValid())
+		{
+			return false;
+		}
+
+		Dim = InDim;
+		Duration = InDuration;
+		OutputDim = InOutputDim;
+		Texture = InTexture;
+		Time = InTime;
+
+		return true;
+	}
+
+#else
+
+	/**
+	 * Initialize the texture sample with a memory buffer.
+	 *
+	 * @param InBuffer The memory buffer containing the video frame.
+	 * @param InDim The sample buffer's width and height (in pixels).
+	 * @param InOutputDim The sample's output width and height (in pixels).
+	 * @param InStride Number of bytes per pixel row.
+	 * @param InTime The sample time (in the player's local clock).
+	 * @param InDuration The duration for which the sample is valid.
+	 * @return true on success, false otherwise.
+	 */
+	bool Initialize(
+		uint8* InBuffer,
+		const FIntPoint& InDim,
+		const FIntPoint& InOutputDim,
+		uint32 InStride,
+		FTimespan InTime,
+		FTimespan InDuration)
+	{
+		if (InBuffer == nullptr)
+		{
+			return false;
+		}
+
+		const int32 BufferSize = InDim.X * InDim.Y * 4;
+
+		if (BufferSize == 0)
+		{
+			return false;
+		}
+
+		Buffer.Reset(BufferSize);
+		Buffer.Append((uint8*)InBuffer, BufferSize);
+
+		Dim = InDim;
+		Duration = InDuration;
+		OutputDim = InOutputDim;
+		Stride = InStride;
+		Time = InTime;
+
+		return true;
+	}
+#endif
+
+public:
+
+	//~ IMediaTextureSample interface
+
+	virtual const void* GetBuffer() override
+	{
+		return nullptr;
+	}
+
+	virtual FIntPoint GetDim() const override
+	{
+		return Texture->GetTexture2D()->GetSizeXY();
+	}
+
+	virtual FTimespan GetDuration() const override
+	{
+		return Duration;
+	}
+
+	virtual EMediaTextureSampleFormat GetFormat() const override
+	{
+		return EMediaTextureSampleFormat::CharBGRA;
+	}
+
+	virtual FIntPoint GetOutputDim() const override
+	{
+		return Texture->GetTexture2D()->GetSizeXY();
+	}
+
+	virtual uint32 GetStride() const override
+	{
+		return Texture->GetTexture2D()->GetSizeX() * 4;
+	}
+
+	virtual FRHITexture* GetTexture() const override
+	{
+		return Texture;
+	}
+
+	virtual FMediaTimeStamp GetTime() const override
+	{
+		return FMediaTimeStamp(Time);
+	}
+
+	virtual bool IsCacheable() const override
+	{
+		return true;
+	}
+
+	virtual bool IsOutputSrgb() const override
+	{
+		return true;
+	}
+
+	virtual void Reset() override
+	{
+		Texture = nullptr;
+	}
+	
+private:
+
+	/** Width and height of the texture sample. */
+	FIntPoint Dim;
+
+	/** Duration for which the sample is valid. */
+	FTimespan Duration;
+
+	/** Width and height of the output. */
+	FIntPoint OutputDim;
+
+	/** Sample time. */
+	FTimespan Time;
+
+#if WITH_ENGINE
+	/** The sample's texture resource. */
+	TRefCountPtr<FRHITexture2D> Texture;
+
+#else
+	/** The sample's data buffer. */
+	TArray<uint8> Buffer;
+
+	/** Number of bytes per pixel row. */
+	uint32 Stride;
+
+#endif
+};
+
+
+/** Implements a pool for AVF texture sample objects. */
+class FAvfMediaTextureSamplePool : public TMediaObjectPool<FAvfMediaTextureSample> { };
+
+/**
+ * Passes a CV*TextureRef or CVPixelBufferRef through to the RHI to wrap in an RHI texture without traversing system memory.
+ */
+class FAvfTexture2DResourceWrapper
+	: public FResourceBulkDataInterface
+{
+public:
+
+	FAvfTexture2DResourceWrapper(CFTypeRef InImageBuffer)
+		: ImageBuffer(InImageBuffer)
+	{
+		check(ImageBuffer);
+		CFRetain(ImageBuffer);
+	}
+
+	virtual ~FAvfTexture2DResourceWrapper()
+	{
+		CFRelease(ImageBuffer);
+		ImageBuffer = nullptr;
+	}
+
+public:
+
+	//~ FResourceBulkDataInterface interface
+
+	virtual void Discard() override
+	{
+		delete this;
+	}
+
+	virtual const void* GetResourceBulkData() const override
+	{
+		return ImageBuffer;
+	}
+	
+	virtual uint32 GetResourceBulkDataSize() const override
+	{
+		return ImageBuffer ? ~0u : 0;
+	}
+
+	virtual EBulkDataType GetResourceType() const override
+	{
+		return EBulkDataType::MediaTexture;
+	}
+			
+	CFTypeRef ImageBuffer;
+};
diff --git a/Engine/Plugins/Media/AvfMedia/Source/AvfMediaCapture/AvfMediaCapture.Build.cs b/Engine/Plugins/Media/AvfMedia/Source/AvfMediaCapture/AvfMediaCapture.Build.cs
new file mode 100644
--- /dev/null
+++ b/Engine/Plugins/Media/AvfMedia/Source/AvfMediaCapture/AvfMediaCapture.Build.cs
@@ -0,0 +1,48 @@
+// Copyright Epic Games, Inc. All Rights Reserved.
+
+using System;
+using System.IO;
+namespace UnrealBuildTool.Rules
+{
+	public class AvfMediaCapture : ModuleRules
+	{
+		public AvfMediaCapture(ReadOnlyTargetRules Target) : base(Target)
+		{
+			DynamicallyLoadedModuleNames.AddRange(
+			new string[] {
+				"Media",
+			});
+
+			PrivateDependencyModuleNames.AddRange(
+				new string[] {
+					"Core",
+					"Engine",
+					"ApplicationCore",
+					"MediaUtils",
+					"RenderCore",
+					"RHI"
+				});
+			
+			PrivateIncludePathModuleNames.AddRange(
+				new string[] {
+					"Media",
+				});
+			
+			var EngineDir = Path.GetFullPath(Target.RelativeEnginePath);
+			PublicIncludePaths.AddRange(
+				new string[] {
+					Path.Combine(EngineDir, "Plugins/Media/AvfMedia/Source/AvfMedia/Public/")
+				});
+			
+			PublicFrameworks.AddRange(
+				new string[] {
+					"CoreMedia",
+					"CoreVideo",
+					"AVFoundation",
+					"AudioToolbox",
+					"MediaToolbox",
+					"QuartzCore"
+				});
+		}
+	}
+}
diff --git a/Engine/Plugins/Media/AvfMedia/Source/AvfMediaCapture/Private/AvfMediaCaptureModule.cpp b/Engine/Plugins/Media/AvfMedia/Source/AvfMediaCapture/Private/AvfMediaCaptureModule.cpp
new file mode 100644
--- /dev/null
+++ b/Engine/Plugins/Media/AvfMedia/Source/AvfMediaCapture/Private/AvfMediaCaptureModule.cpp
@@ -0,0 +1,203 @@
+// Copyright Epic Games, Inc. All Rights Reserved.
+
+#include "AvfMediaCapturePrivate.h"
+
+#include "Containers/Array.h"
+#include "Containers/UnrealString.h"
+#include "IMediaModule.h"
+#include "IMediaPlayerFactory.h"
+#include "Internationalization/Internationalization.h"
+#include "Misc/Paths.h"
+#include "Modules/ModuleInterface.h"
+#include "Modules/ModuleManager.h"
+#include "UObject/NameTypes.h"
+#include "IMediaCaptureSupport.h"
+#include "Player/AvfMediaCapturePlayer.h"
+
+#if WITH_EDITOR
+	#include "ISettingsModule.h"
+	#include "ISettingsSection.h"
+	#include "UObject/Class.h"
+	#include "UObject/WeakObjectPtr.h"
+#endif
+
+#import <AVFoundation/AVFoundation.h>
+
+
+DEFINE_LOG_CATEGORY(LogAvfMediaCapture);
+
+#define LOCTEXT_NAMESPACE "FAvfMediaCaptureFactoryModule"
+
+
+/**
+ * Implements the AvfMediaCapture module.
+ */
+class FAvfMediaCaptureModule
+	: public IModuleInterface
+	, public IMediaPlayerFactory
+	, public IMediaCaptureSupport
+{
+public:
+	void EnumerateCaptureDevices(TArray<FMediaCaptureDeviceInfo>& OutDeviceInfos, EMediaCaptureDeviceType TargetDeviceType)
+	{
+		SCOPED_AUTORELEASE_POOL
+
+		NSArray* DeviceTypes = @[AVCaptureDeviceTypeBuiltInWideAngleCamera, AVCaptureDeviceTypeBuiltInMicrophone];
+
+		AVCaptureDeviceDiscoverySession* LocalDiscoverySession = [AVCaptureDeviceDiscoverySession discoverySessionWithDeviceTypes:DeviceTypes mediaType:nil position:AVCaptureDevicePositionUnspecified];
+		if(LocalDiscoverySession != nil)
+		{
+			NSArray<AVCaptureDevice*>* Devices = LocalDiscoverySession.devices;
+			for(uint32 i = 0;i < Devices.count;++i)
+			{
+				AVCaptureDevice* AvailableDevice = Devices[i];
+				
+				if(TargetDeviceType == EMediaCaptureDeviceType::Audio && AvailableDevice.deviceType == AVCaptureDeviceTypeBuiltInMicrophone)
+				{
+					FMediaCaptureDeviceInfo DeviceInfo;
+					
+					DeviceInfo.Type = TargetDeviceType;
+					DeviceInfo.DisplayName = FText::FromString(FString(AvailableDevice.localizedName));
+					DeviceInfo.Url = TEXT("audcap://") + FString(AvailableDevice.uniqueID);
+					DeviceInfo.Info = FString(AvailableDevice.manufacturer);
+					
+					OutDeviceInfos.Add(MoveTemp(DeviceInfo));
+				}
+				else if(TargetDeviceType == EMediaCaptureDeviceType::Video && AvailableDevice.deviceType != AVCaptureDeviceTypeBuiltInMicrophone)
+				{
+					FMediaCaptureDeviceInfo DeviceInfo;
+				
+					DeviceInfo.Type = TargetDeviceType;
+					DeviceInfo.DisplayName = FText::FromString(FString(AvailableDevice.localizedName));
+					DeviceInfo.Url = TEXT("vidcap://") + FString(AvailableDevice.uniqueID);
+					DeviceInfo.Info = FString(AvailableDevice.manufacturer);
+					
+					OutDeviceInfos.Add(MoveTemp(DeviceInfo));
+				}
+			}
+		}
+	}
+
+	//~ IMediaCaptureSupport interface
+	
+	virtual void EnumerateAudioCaptureDevices(TArray<FMediaCaptureDeviceInfo>& OutDeviceInfos)
+	{
+		EnumerateCaptureDevices(OutDeviceInfos, EMediaCaptureDeviceType::Audio);
+	}
+	
+	virtual void EnumerateVideoCaptureDevices(TArray<FMediaCaptureDeviceInfo>& OutDeviceInfos)
+	{
+		EnumerateCaptureDevices(OutDeviceInfos, EMediaCaptureDeviceType::Video);
+	}
+
+	//~ IMediaPlayerFactory interface
+
+	virtual bool CanPlayUrl(const FString& Url, const IMediaOptions* /*Options*/, TArray<FText>* /*OutWarnings*/, TArray<FText>* OutErrors) const override
+	{
+		FString Scheme;
+		FString Location;
+
+		// check scheme
+		if (!Url.Split(TEXT("://"), &Scheme, &Location, ESearchCase::CaseSensitive))
+		{
+			if (OutErrors != nullptr)
+			{
+				OutErrors->Add(LOCTEXT("NoSchemeFound", "No URI scheme found"));
+			}
+
+			return false;
+		}
+
+		if (!SupportedUriSchemes.Contains(Scheme))
+		{
+			if (OutErrors != nullptr)
+			{
+				OutErrors->Add(FText::Format(LOCTEXT("SchemeNotSupported", "The URI scheme '{0}' is not supported"), FText::FromString(Scheme)));
+			}
+
+			return false;
+		}
+
+		return true;
+	}
+
+	virtual TSharedPtr<IMediaPlayer, ESPMode::ThreadSafe> CreatePlayer(IMediaEventSink& EventSink) override
+	{
+		return MakeShared<FAvfMediaCapturePlayer, ESPMode::ThreadSafe>(EventSink);
+	}
+
+	virtual FText GetDisplayName() const override
+	{
+		return LOCTEXT("MediaCaptureDisplayName", "Apple AV Foundation");
+	}
+
+	virtual FName GetPlayerName() const override
+	{
+		static FName PlayerName(TEXT("AvfMediaCapture"));
+		return PlayerName;
+	}
+
+	virtual FGuid GetPlayerPluginGUID() const override
+	{
+		static FGuid PlayerPluginGUID(0xcf78bfd2, 0x0c1111ed, 0x861d0242, 0xac120002);
+		return PlayerPluginGUID;
+	}
+
+	virtual const TArray<FString>& GetSupportedPlatforms() const override
+	{
+		return SupportedPlatforms;
+	}
+
+	virtual bool SupportsFeature(EMediaFeature Feature) const override
+	{
+		return ((Feature == EMediaFeature::AudioSamples) ||
+				(Feature == EMediaFeature::VideoSamples));
+	}
+
+public:
+
+	//~ IModuleInterface interface
+
+	virtual void StartupModule() override
+	{
+		// supported schemes
+		SupportedUriSchemes.Add(TEXT("vidcap"));
+		SupportedUriSchemes.Add(TEXT("audcap"));
+
+		// supported platforms
+		SupportedPlatforms.Add(TEXT("Mac"));
+		SupportedPlatforms.Add(TEXT("iOS"));
+
+		// register factory support functions
+		auto MediaModule = FModuleManager::LoadModulePtr<IMediaModule>("Media");
+		if (MediaModule != nullptr)
+		{
+			MediaModule->RegisterPlayerFactory(*this);
+			MediaModule->RegisterCaptureSupport(*this);
+		}
+	}
+
+	virtual void ShutdownModule() override
+	{
+		// unregister factory support functions
+		auto MediaModule = FModuleManager::GetModulePtr<IMediaModule>("Media");
+		if (MediaModule != nullptr)
+		{
+			MediaModule->UnregisterPlayerFactory(*this);
+			MediaModule->UnregisterCaptureSupport(*this);
+		}
+	}
+
+private:
+
+	/** List of platforms that the media player support. */
+	TArray<FString> SupportedPlatforms;
+
+	/** List of supported URI schemes. */
+	TArray<FString> SupportedUriSchemes;
+};
+
+
+#undef LOCTEXT_NAMESPACE
+
+IMPLEMENT_MODULE(FAvfMediaCaptureModule, AvfMediaCapture);
diff --git a/Engine/Plugins/Media/AvfMedia/Source/AvfMediaCapture/Private/AvfMediaCapturePrivate.h b/Engine/Plugins/Media/AvfMedia/Source/AvfMediaCapture/Private/AvfMediaCapturePrivate.h
new file mode 100644
--- /dev/null
+++ b/Engine/Plugins/Media/AvfMedia/Source/AvfMediaCapture/Private/AvfMediaCapturePrivate.h
@@ -0,0 +1,9 @@
+// Copyright Epic Games, Inc. All Rights Reserved.
+
+#pragma once
+
+#include "Logging/LogMacros.h"
+
+
+/** Log category for the AvfMediaCapture module. */
+DECLARE_LOG_CATEGORY_EXTERN(LogAvfMediaCapture, Log, All);
diff --git a/Engine/Plugins/Media/AvfMedia/Source/AvfMediaCapture/Private/Player/AvfMediaCaptureHelper.cpp b/Engine/Plugins/Media/AvfMedia/Source/AvfMediaCapture/Private/Player/AvfMediaCaptureHelper.cpp
new file mode 100644
--- /dev/null
+++ b/Engine/Plugins/Media/AvfMedia/Source/AvfMediaCapture/Private/Player/AvfMediaCaptureHelper.cpp
@@ -0,0 +1,396 @@
+// Copyright Epic Games, Inc. All Rights Reserved.
+
+#import <Foundation/Foundation.h>
+#import <AVFoundation/AVFoundation.h>
+#include "AvfMediaCaptureHelper.h"
+
+// Steps to try if permissions are not working correctly in macOS
+
+// 1) macOS resetting permissions, also useful for testing:
+//  tccutil reset Camera
+//  tccutil reset Microphone
+//
+// 2) System integrity protection has to be enabled for permissions to work correctly, check using:
+//  csrutil status
+//
+// 3) delete permissions database then reboot mac:
+//  ~/Library/Application\\ Support/com.apple.TCC
+
+@interface AvfMediaCaptureHelper()<AVCaptureVideoDataOutputSampleBufferDelegate, AVCaptureAudioDataOutputSampleBufferDelegate>
+
+@property (nonatomic,readwrite,assign) AVCaptureSession* 		captureSession;
+@property (nonatomic,readwrite,assign) AVCaptureDevice*			captureDevice;
+@property (nonatomic,readwrite,assign) AVCaptureDeviceInput*	deviceInput;
+@property (nonatomic,readwrite,assign) AVCaptureOutput* 		captureOutput;
+@property (nonatomic,readwrite,assign) dispatch_queue_t 		callBackQueue;
+
+@property (nonatomic,readwrite,assign) void (^sampleBufferCallback)(CMSampleBufferRef);
+@property (nonatomic,readwrite,assign) void (^notificationCallback)(NSNotification* const Notification);
+
+@end
+
+@implementation AvfMediaCaptureHelper
+
+- (instancetype) init
+{
+	self = [super init];
+	if(self != nil)
+	{
+		self.captureSession = nil;
+		self.captureDevice = nil;
+		self.deviceInput = nil;
+		self.captureOutput = nil;
+		self.sampleBufferCallback = nil;
+		self.notificationCallback = nil;
+		self.callBackQueue = dispatch_queue_create("AvfCaptureSessionSampleCallbackQueue", DISPATCH_QUEUE_SERIAL);
+	}
+	return self;
+}
+
+- (void) dealloc
+{
+	[self reset];
+	
+	if(self.callBackQueue != nil)
+	{
+		dispatch_release(self.callBackQueue);
+		self.callBackQueue = nil;
+	}
+	
+	[super dealloc];
+}
+
+- (void) reset
+{
+	if(self.captureSession.inputs.count > 0 && self.captureSession.outputs.count > 0)
+	{
+		[[NSNotificationCenter defaultCenter] removeObserver:self];
+	}
+	
+	if(self.captureSession.isRunning)
+	{
+		[self.captureSession stopRunning];
+	}
+	
+	if(self.captureOutput != nil)
+	{
+		[self.captureSession removeOutput:self.captureOutput];
+		[self.captureOutput release];
+		self.captureOutput = nil;
+	}
+	
+	if(self.deviceInput != nil)
+	{
+		[self.captureSession removeInput:self.deviceInput];
+		[self.deviceInput release];
+		self.deviceInput = nil;
+	}
+	
+	if(self.captureSession != nil)
+	{
+		[self.captureSession release];
+		self.captureSession = nil;
+	}
+	
+	if(self.captureDevice != nil)
+	{
+		[self.captureDevice release];
+		self.captureDevice = nil;
+	}
+	
+	if(self.sampleBufferCallback != nil)
+	{
+		Block_release(self.sampleBufferCallback);
+		self.sampleBufferCallback = nil;
+	}
+	
+	if(self.notificationCallback != nil)
+	{
+		Block_release(self.notificationCallback);
+		self.notificationCallback = nil;
+	}
+}
+
++ (EAvfMediaCaptureAuthStatus) authorizationStatusForMediaType:(AVMediaType)mediaType
+{
+	if(mediaType != AVMediaTypeVideo && mediaType != AVMediaTypeAudio)
+	{
+		return EAvfMediaCaptureAuthStatus::InvalidRequest;
+	}
+
+	NSDictionary<NSString*,id>* infoDictionary = [NSBundle mainBundle].infoDictionary;
+	
+	id entry = nil;
+	
+	if(infoDictionary != nil)
+	{
+		if(mediaType == AVMediaTypeVideo)
+		{
+			entry = infoDictionary[@"NSCameraUsageDescription"];
+		}
+		else if (mediaType == AVMediaTypeAudio)
+		{
+			entry = infoDictionary[@"NSMicrophoneUsageDescription"];
+		}
+	}
+	
+	if(entry == nil)
+	{
+		return EAvfMediaCaptureAuthStatus::MissingInfoPListEntry;
+	}
+	
+	return (EAvfMediaCaptureAuthStatus)[AVCaptureDevice authorizationStatusForMediaType:mediaType];
+}
+
++ (EAvfMediaCaptureAuthStatus) requestAcessForMediaType:(AVMediaType)mediaType completionCallback:(void (^)(EAvfMediaCaptureAuthStatus AuthStatus))cbHandler
+{
+	// Don't make request if the correct device access key is not in the info.plist otherwise the OS will terminate this app
+	EAvfMediaCaptureAuthStatus authStatus = [AvfMediaCaptureHelper authorizationStatusForMediaType:mediaType];
+	switch(authStatus)
+	{
+		case EAvfMediaCaptureAuthStatus::NotDetermined:
+		{
+			[AVCaptureDevice requestAccessForMediaType:mediaType completionHandler:^(BOOL bGranted)
+			{
+				cbHandler((EAvfMediaCaptureAuthStatus)[AVCaptureDevice authorizationStatusForMediaType:mediaType]);
+			}];
+			break;
+		}
+		case EAvfMediaCaptureAuthStatus::Authorized:
+		case EAvfMediaCaptureAuthStatus::Restricted:
+		case EAvfMediaCaptureAuthStatus::Denied:
+		case EAvfMediaCaptureAuthStatus::MissingInfoPListEntry:
+		case EAvfMediaCaptureAuthStatus::InvalidRequest:
+		{
+			break;
+		}
+	}
+	
+	return authStatus;
+}
+
+- (AVCaptureDevice*) CaptureDeviceWithID:(NSString*)deviceID
+{
+	AVCaptureDevice* foundDevice = nil;
+		
+	NSArray* deviceTypes = @[
+								AVCaptureDeviceTypeBuiltInWideAngleCamera,
+								AVCaptureDeviceTypeBuiltInMicrophone
+							];
+
+	AVCaptureDeviceDiscoverySession* localDiscoverySession = [AVCaptureDeviceDiscoverySession discoverySessionWithDeviceTypes:deviceTypes
+																													mediaType:nil
+																													 position:AVCaptureDevicePositionUnspecified];
+	if(localDiscoverySession != nil)
+	{
+		NSArray<AVCaptureDevice*>* devices = localDiscoverySession.devices;
+		
+		for(AVCaptureDevice* device in devices)
+		{
+			if([device.uniqueID isEqual:deviceID])
+			{
+				foundDevice = device;
+				break;
+			}
+		}
+	}
+
+	return foundDevice;
+}
+
+- (BOOL) setupCaptureSession:(NSString*)deviceID sampleBufferCallback:(void(^)(CMSampleBufferRef sampleBuffer))sampleCallbackBlock notificationCallback:(void(^)(NSNotification* const notification))notificationCallbackBlock
+{
+	[self reset];
+	
+	if(sampleCallbackBlock != nil)
+	{
+		self.sampleBufferCallback = Block_copy(sampleCallbackBlock);
+	}
+	if(notificationCallbackBlock != nil)
+	{
+		self.notificationCallback = Block_copy(notificationCallbackBlock);
+	}
+	
+	self.captureDevice = [[self CaptureDeviceWithID:deviceID] retain];
+	if(self.captureDevice != nil)
+	{
+		self.captureSession = [[AVCaptureSession alloc] init];
+		
+		// Default to high so it starts with a reasonably good initial setting
+		self.captureSession.sessionPreset = AVCaptureSessionPresetHigh;
+		
+		// Add device input
+		{
+			self.deviceInput = [[AVCaptureDeviceInput alloc] initWithDevice:self.captureDevice error:nil];
+			
+			if(self.deviceInput != nil && [self.captureSession canAddInput:self.deviceInput])
+			{
+				[self.captureSession addInput:self.deviceInput];
+			}
+		}
+		
+		// If we have an input add output capture
+		if(self.captureSession.inputs.count > 0)
+		{
+			// Video or Audio Device - need the correct output
+			if(self.captureDevice.deviceType != AVCaptureDeviceTypeBuiltInMicrophone)
+			{
+				AVCaptureVideoDataOutput* videoOutput = [[AVCaptureVideoDataOutput alloc] init];
+				
+				[videoOutput setSampleBufferDelegate:self queue:self.callBackQueue];
+				videoOutput.alwaysDiscardsLateVideoFrames = YES;
+				
+				videoOutput.videoSettings = @{ 	(id)kCVPixelBufferPixelFormatTypeKey 	: @(kCVPixelFormatType_32BGRA),
+												(id)kCVPixelBufferMetalCompatibilityKey : @(1)};
+				
+				self.captureOutput = videoOutput;
+			}
+			else
+			{
+				AVCaptureAudioDataOutput* audioOutput = [[AVCaptureAudioDataOutput alloc] init];
+				
+				[audioOutput setSampleBufferDelegate:self queue:self.callBackQueue];
+
+#define MACOS_AUDIO_CAPTURE_USR_SETTINGS 0
+#if PLATFORM_MAC && MACOS_AUDIO_CAPTURE_USR_SETTINGS != 0
+				// AVCaptureAudioDataOutput.audioSettings is available on macOS only
+				// Use default device / OS output instead.  This is left here for debug and test purposes
+				// Changes the output audio format delivered to the sample buffer delete function
+				// See AVFAudio/AVAudioSettings.h for available keys / values
+				audioOutput.audioSettings = @{
+												// LPCM=Linear Quantization Levels
+												AVFormatIDKey					:@(kAudioFormatLinearPCM),
+												// 8/16.24/32=INT, 32=FLOAT
+												AVLinearPCMBitDepthKey 			: @(32),
+												// NO=INT, YES=FLOAT, engine perfers float and will convert to float eventually what ever we output
+												AVLinearPCMIsFloatKey 			: @YES,
+												// Non Interleaved each channel is sequential one after each other e.g 3 channel samples = 111222333, Interleaved e.g. 123123123, engine wants interleaved
+												AVLinearPCMIsNonInterleavedKey	: @NO,
+												// As Per hardware
+												AVNumberOfChannelsKey			: @(1)
+											};
+#endif
+#undef MACOS_AUDIO_CAPTURE_USR_SETTINGS
+
+				self.captureOutput = audioOutput;
+			}
+			
+			if(self.captureOutput != nil && [self.captureSession canAddOutput:self.captureOutput])
+			{
+				[self.captureSession addOutput:self.captureOutput];
+			}
+		}
+	}
+	
+	if(self.captureSession.inputs.count > 0 && self.captureSession.outputs.count > 0)
+	{
+		NSNotificationCenter* sharedCentre = [NSNotificationCenter defaultCenter];
+
+		[sharedCentre addObserver:self selector:@selector(captureNotification:) name:AVCaptureSessionDidStartRunningNotification object:self.captureSession];
+		[sharedCentre addObserver:self selector:@selector(captureNotification:) name:AVCaptureSessionDidStopRunningNotification object:self.captureSession];
+		[sharedCentre addObserver:self selector:@selector(captureNotification:) name:AVCaptureSessionWasInterruptedNotification object:self.captureSession];
+		[sharedCentre addObserver:self selector:@selector(captureNotification:) name:AVCaptureSessionInterruptionEndedNotification object:self.captureSession];
+		[sharedCentre addObserver:self selector:@selector(captureNotification:) name:AVCaptureSessionRuntimeErrorNotification object:self.captureSession];
+		
+#if PLATFORM_MAC && WITH_EDITOR
+		if([self.captureOutput isKindOfClass:[AVCaptureVideoDataOutput class]])
+		{
+			[sharedCentre addObserver:self selector:@selector(captureNotification:) name:NSApplicationDidBecomeActiveNotification object:nil];
+			[sharedCentre addObserver:self selector:@selector(captureNotification:) name:NSApplicationWillResignActiveNotification object:nil];
+		}
+#endif
+
+		return true;
+	}
+
+	return false;
+}
+
+- (void) stopCaptureSession
+{
+	if(self.captureSession.isRunning)
+	{
+		[self.captureSession stopRunning];
+	}
+}
+
+- (void) startCaptureSession
+{
+	if(!self.captureSession.isRunning)
+	{
+		[self.captureSession startRunning];
+	}
+}
+
+- (BOOL) isCaptureRunning
+{
+	return self.captureSession.isRunning;
+}
+
+- (void)captureOutput:(AVCaptureOutput *)output didOutputSampleBuffer:(CMSampleBufferRef)sampleBuffer fromConnection:(AVCaptureConnection*)connection
+{
+	if(self.sampleBufferCallback != nil && sampleBuffer != NULL)
+	{
+		self.sampleBufferCallback(sampleBuffer);
+	}
+}
+
+- (void) captureNotification:(NSNotification*)notification
+{
+	if(self.notificationCallback != nil && notification != nil)
+	{
+		self.notificationCallback(notification);
+	}
+}
+
+- (AVMediaType) getCaptureDeviceMediaType
+{
+	AVMediaType type = nil;
+	if(self.captureDevice != nil)
+	{
+		if(self.captureDevice.deviceType != AVCaptureDeviceTypeBuiltInMicrophone)
+		{
+			type = AVMediaTypeVideo;
+		}
+		else
+		{
+			type = AVMediaTypeAudio;
+		}
+	}
+	return type;
+}
+
+- (NSString*) getCaptureDeviceName
+{
+	return self.captureDevice.localizedName;
+}
+
+- (NSArray<AVCaptureDeviceFormat*>*) getCaptureDeviceAvailableFormats
+{
+	return self.captureDevice.formats;
+}
+
+- (NSInteger) getCaptureDeviceActiveFormatIndex
+{
+	for(NSInteger formatIndex = 0;formatIndex < self.captureDevice.formats.count;++formatIndex)
+	{
+		if([self.captureDevice.formats[formatIndex] isEqual:self.captureDevice.activeFormat])
+		{
+			return formatIndex;
+		}
+	}
+	return INDEX_NONE;
+}
+
+- (BOOL) setCaptureDeviceActiveFormatIndex:(NSInteger)formatIdx
+{
+	if(formatIdx >= 0 && formatIdx < self.captureDevice.formats.count && [self.captureDevice lockForConfiguration:nil])
+	{
+		self.captureDevice.activeFormat = self.captureDevice.formats[formatIdx];
+		[self.captureDevice unlockForConfiguration];
+		return YES;
+	}
+	return NO;
+}
+
+@end
diff --git a/Engine/Plugins/Media/AvfMedia/Source/AvfMediaCapture/Private/Player/AvfMediaCaptureHelper.h b/Engine/Plugins/Media/AvfMedia/Source/AvfMediaCapture/Private/Player/AvfMediaCaptureHelper.h
new file mode 100644
--- /dev/null
+++ b/Engine/Plugins/Media/AvfMedia/Source/AvfMediaCapture/Private/Player/AvfMediaCaptureHelper.h
@@ -0,0 +1,42 @@
+// Copyright Epic Games, Inc. All Rights Reserved.
+
+#pragma once
+
+#import <AVFoundation/AVFoundation.h>
+
+enum class EAvfMediaCaptureAuthStatus : uint32
+{
+	// Mapped Apple AVFoundation values
+    NotDetermined	= AVAuthorizationStatusNotDetermined,
+	Restricted		= AVAuthorizationStatusRestricted,
+	Denied			= AVAuthorizationStatusDenied,
+	Authorized		= AVAuthorizationStatusAuthorized,
+    
+    // Extra Error values
+	MissingInfoPListEntry,
+	InvalidRequest
+};
+
+@interface AvfMediaCaptureHelper : NSObject
+
+// Returns current status or error
++ (EAvfMediaCaptureAuthStatus) authorizationStatusForMediaType:(AVMediaType)mediaType;
+
+// Returns current status or error and requests access if not determined
++ (EAvfMediaCaptureAuthStatus) requestAcessForMediaType:(AVMediaType)mediaType completionCallback:(void(^)(EAvfMediaCaptureAuthStatus AuthStatus))cbHandler;
+
+- (BOOL) setupCaptureSession:(NSString*)deviceID sampleBufferCallback:(void(^)(CMSampleBufferRef sampleBuffer))sampleCallbackBlock
+												 notificationCallback:(void(^)(NSNotification* const notification))notificationCallbackBlock;
+- (void) stopCaptureSession;
+- (void) startCaptureSession;
+- (BOOL) isCaptureRunning;
+
+- (NSString*) 	getCaptureDeviceName;
+- (AVMediaType) getCaptureDeviceMediaType;
+
+- (NSArray<AVCaptureDeviceFormat*>*) getCaptureDeviceAvailableFormats;
+
+- (NSInteger) 	getCaptureDeviceActiveFormatIndex;
+- (BOOL) 		setCaptureDeviceActiveFormatIndex:(NSInteger)formatIdx;
+
+@end
diff --git a/Engine/Plugins/Media/AvfMedia/Source/AvfMediaCapture/Private/Player/AvfMediaCapturePlayer.cpp b/Engine/Plugins/Media/AvfMedia/Source/AvfMediaCapture/Private/Player/AvfMediaCapturePlayer.cpp
new file mode 100644
--- /dev/null
+++ b/Engine/Plugins/Media/AvfMedia/Source/AvfMediaCapture/Private/Player/AvfMediaCapturePlayer.cpp
@@ -0,0 +1,851 @@
+// Copyright Epic Games, Inc. All Rights Reserved.
+
+#include "AvfMediaCapturePlayer.h"
+#include "MediaSamples.h"
+#include "Misc/MessageDialog.h"
+
+inline FString FStringFromAppleFourCharWordCode(uint32 Word)
+{
+	char TypeCode[5];
+	{
+		TypeCode[0] = ((char*)&Word)[3];
+		TypeCode[1] = ((char*)&Word)[2];
+		TypeCode[2] = ((char*)&Word)[1];
+		TypeCode[3] = ((char*)&Word)[0];
+		TypeCode[4] = 0;
+	}
+	return FString((char*)(&TypeCode));
+}
+
+template<class T>
+void InterleaveCopy(void* DestPtr, void* SrcPtr, size_t NumChannels, size_t NumFrames, size_t SrcFrameOffset, size_t SrcNumFrames)
+{
+	T* DestPtrType = (T*)DestPtr;
+	T* SrcPtrType = (T*)SrcPtr;
+
+	for(size_t SrcFrameIdx = 0;SrcFrameIdx < SrcNumFrames;++SrcFrameIdx)
+	{
+		size_t GlobalSrcFrameIdx = SrcFrameOffset + SrcFrameIdx;
+		size_t ChannelIdx = GlobalSrcFrameIdx / NumFrames;
+		size_t DestFrameIdx = (NumChannels * (GlobalSrcFrameIdx % NumFrames)) + ChannelIdx;
+
+		DestPtrType[DestFrameIdx] = SrcPtrType[SrcFrameIdx];
+	}
+}
+
+FAvfMediaCapturePlayer::FAvfMediaCapturePlayer(IMediaEventSink& InEventSink)
+: EventSink(InEventSink)
+, MediaSamples(MakeShared<FMediaSamples, ESPMode::ThreadSafe>())
+, MediaCaptureHelper(nil)
+, MetalTextureCache(NULL)
+, CurrentRate(0.f)
+, CurrentTime(FTimespan::MinValue())
+#if PLATFORM_MAC && WITH_EDITOR
+, ThrottleDuration(0.0)
+, LastConsumedTimeStamp(0.0)
+#endif
+{}
+
+FAvfMediaCapturePlayer::~FAvfMediaCapturePlayer()
+{
+	Close();
+	
+	if(MetalTextureCache != NULL)
+	{
+		CFRelease(MetalTextureCache);
+		MetalTextureCache = NULL;
+	}
+}
+
+// IMediaPlayer Interface
+
+IMediaCache& FAvfMediaCapturePlayer::GetCache()
+{
+	return *this;
+}
+
+IMediaControls& FAvfMediaCapturePlayer::GetControls()
+{
+	return *this;
+}
+
+FString FAvfMediaCapturePlayer::GetInfo() const
+{
+	return FString("");
+}
+
+FGuid FAvfMediaCapturePlayer::GetPlayerPluginGUID() const
+{
+	static FGuid PlayerPluginGUID(0xcf78bfd2, 0x0c1111ed, 0x861d0242, 0xac120002);
+	return PlayerPluginGUID;
+}
+
+IMediaSamples& FAvfMediaCapturePlayer::GetSamples()
+{
+	return *MediaSamples.Get();
+}
+
+FString FAvfMediaCapturePlayer::GetStats() const
+{
+	return FString("");
+}
+
+IMediaTracks& FAvfMediaCapturePlayer::GetTracks()
+{
+	return *this;
+}
+
+FString FAvfMediaCapturePlayer::GetUrl() const
+{
+	return URL;
+}
+
+FText FAvfMediaCapturePlayer::GetMediaName() const
+{
+	FScopeLock Lock(&CriticalSection);
+	if(MediaCaptureHelper != nil)
+	{
+		return FText::FromString(FString([MediaCaptureHelper getCaptureDeviceName]));
+	}
+	return FText();
+}
+
+IMediaView& FAvfMediaCapturePlayer::GetView()
+{
+	return *this;
+}
+
+void FAvfMediaCapturePlayer::Close()
+{
+	SCOPED_AUTORELEASE_POOL;
+	
+	FScopeLock Lock(&CriticalSection);
+	
+	CurrentRate = 0.f;
+	CurrentTime = FTimespan::MinValue();
+	URL = FString();
+	
+#if PLATFORM_MAC && WITH_EDITOR
+	ThrottleDuration = 0.0;
+	LastConsumedTimeStamp = 0.0;
+#endif
+	
+	if(MediaCaptureHelper != nil)
+	{
+		[MediaCaptureHelper release];
+		MediaCaptureHelper = nil;
+		
+		EventSink.ReceiveMediaEvent(EMediaEvent::MediaClosed);
+	}
+}
+
+void FAvfMediaCapturePlayer::HandleAuthStatusError(EAvfMediaCaptureAuthStatus AuthStatus, AVMediaType MediaType)
+{
+	FString ErrorString;
+	
+	switch(AuthStatus)
+	{
+		case EAvfMediaCaptureAuthStatus::Authorized:
+		case EAvfMediaCaptureAuthStatus::NotDetermined:
+		{
+			// NO ERROR
+			break;
+		}
+		case EAvfMediaCaptureAuthStatus::MissingInfoPListEntry:
+		{
+			ErrorString = NSLOCTEXT("AvfMediaCapture", "CaptureDeviceFail", "Cannot start capture: Missing Info.plist entry: ").ToString() + (MediaType == AVMediaTypeVideo ? TEXT("\"NSCameraUsageDescription\"") : TEXT("\"NSMicrophoneUsageDescription\""));
+			break;
+		}
+		case EAvfMediaCaptureAuthStatus::Restricted:
+		{
+			ErrorString = NSLOCTEXT("AvfMediaCapture", "CaptureDeviceFail", "Cannot start capture: Restricted").ToString();
+			break;
+		}
+		case EAvfMediaCaptureAuthStatus::Denied:
+		{
+			ErrorString = NSLOCTEXT("AvfMediaCapture", "CaptureDeviceFail", "Cannot start capture: Denied").ToString();
+			break;
+		}
+		default:
+		{
+			ErrorString = NSLOCTEXT("AvfMediaCapture", "CaptureDeviceFail", "Cannot start capture: Unknown reason").ToString();
+			break;
+		}
+	}
+	
+	if(ErrorString.Len() > 0)
+	{
+		EventSink.ReceiveMediaEvent(EMediaEvent::MediaOpenFailed);
+
+		ErrorString += FString(TEXT(" (")) + URL + FString(TEXT(")"));
+		FText Msg = FText::FromString(ErrorString);
+		FText Title = NSLOCTEXT("AvfMediaCapture", "CaptureDeviceFail", "Capture Device Start Failed");
+		FMessageDialog::Open(EAppMsgType::Ok, Msg, &Title);
+	}
+}
+
+void FAvfMediaCapturePlayer::CreateCaptureSession(NSString* deviceIDString)
+{
+	SCOPED_AUTORELEASE_POOL;
+	
+	FScopeLock Lock(&CriticalSection);
+	
+	if(MediaCaptureHelper == nil)
+	{
+		// Assuming only one Capture at the moment
+		MediaCaptureHelper = [[AvfMediaCaptureHelper alloc] init];
+		BOOL bResult = [MediaCaptureHelper setupCaptureSession:deviceIDString sampleBufferCallback:^(CMSampleBufferRef SampleBuffer)
+		{
+			this->NewSampleBufferAvailable(SampleBuffer);
+		}
+		notificationCallback:^(NSNotification* Notification)
+		{
+			this->CaptureSystemNotification(Notification);
+		}];
+		
+		if(bResult)
+		{
+			EventSink.ReceiveMediaEvent(EMediaEvent::MediaOpened);
+		}
+		else
+		{
+			Close();
+		}
+	}
+}
+
+bool FAvfMediaCapturePlayer::Open(const FString& Url, const IMediaOptions* Options)
+{
+	SCOPED_AUTORELEASE_POOL;
+	
+	Close();
+	
+	FScopeLock Lock(&CriticalSection);
+	
+	FString DeviceType;
+	FString DeviceID;
+	
+	Url.Split(TEXT("://"), &DeviceType, &DeviceID);
+	
+	if((DeviceType.StartsWith(TEXT("vidcap")) || DeviceType.StartsWith(TEXT("audcap"))) && DeviceID.Len() > 0)
+	{
+		URL = Url;
+		AVMediaType MediaType = DeviceType.StartsWith(TEXT("vidcap")) ? AVMediaTypeVideo : AVMediaTypeAudio;
+		
+		NSString* nsDeviceID = DeviceID.GetNSString();
+		
+		EAvfMediaCaptureAuthStatus CaptureAuthStatus = [AvfMediaCaptureHelper authorizationStatusForMediaType:MediaType];
+		switch(CaptureAuthStatus)
+		{
+			case EAvfMediaCaptureAuthStatus::Authorized:
+			{
+				CreateCaptureSession(nsDeviceID);
+				break;
+			}
+			case EAvfMediaCaptureAuthStatus::NotDetermined:
+			{
+				TWeakPtr<FAvfMediaCapturePlayer, ESPMode::ThreadSafe> WeakPtr = this->AsWeak();
+				[AvfMediaCaptureHelper requestAcessForMediaType:MediaType completionCallback:^(EAvfMediaCaptureAuthStatus ResultStatus)
+				{
+					TSharedPtr<FAvfMediaCapturePlayer, ESPMode::ThreadSafe> StrongPtr = WeakPtr.Pin();
+					if(StrongPtr.IsValid())
+					{
+						if(ResultStatus == EAvfMediaCaptureAuthStatus::Authorized)
+						{
+							StrongPtr->CreateCaptureSession(nsDeviceID);
+						}
+						else
+						{
+							StrongPtr->HandleAuthStatusError(ResultStatus, MediaType);
+						}
+					}
+				}];
+				break;
+			}
+			case EAvfMediaCaptureAuthStatus::MissingInfoPListEntry:
+			case EAvfMediaCaptureAuthStatus::Restricted:
+			case EAvfMediaCaptureAuthStatus::Denied:
+			default:
+			{
+				HandleAuthStatusError(CaptureAuthStatus, MediaType);
+				break;
+			}
+		}
+		
+		return MediaCaptureHelper != nil || CaptureAuthStatus == EAvfMediaCaptureAuthStatus::NotDetermined;
+	}
+	
+	return false;
+}
+
+bool FAvfMediaCapturePlayer::Open(const TSharedRef<FArchive, ESPMode::ThreadSafe>& Archive, const FString& OriginalUrl, const IMediaOptions* Options)
+{
+	// Does not support open/playback from an Unreal Engine Archive
+	return false;
+}
+
+// IMediaTracks Interface
+
+bool FAvfMediaCapturePlayer::GetAudioTrackFormat(int32 TrackIndex, int32 FormatIndex, FMediaAudioTrackFormat& OutFormat) const
+{
+	SCOPED_AUTORELEASE_POOL;
+	
+	FScopeLock Lock(&CriticalSection);
+	if(TrackIndex == 0 && FormatIndex >= 0)
+	{
+		NSArray<AVCaptureDeviceFormat*>* deviceFormats = [MediaCaptureHelper getCaptureDeviceAvailableFormats];
+		
+		if(FormatIndex < deviceFormats.count)
+		{
+			AVCaptureDeviceFormat* availableFormat = deviceFormats[FormatIndex];
+			
+			if(CMFormatDescriptionGetMediaType(availableFormat.formatDescription) == kCMMediaType_Audio)
+			{
+				CMAudioFormatDescriptionRef AudioFormatDescription = (CMAudioFormatDescriptionRef)availableFormat.formatDescription;
+				AudioStreamBasicDescription const* ASBD = CMAudioFormatDescriptionGetStreamBasicDescription(AudioFormatDescription);
+
+				OutFormat.BitsPerSample = ASBD->mBitsPerChannel;
+				OutFormat.NumChannels = ASBD->mChannelsPerFrame;
+				OutFormat.SampleRate = (uint32)ASBD->mSampleRate;
+				OutFormat.TypeName = FStringFromAppleFourCharWordCode(ASBD->mFormatID);
+				
+				if((ASBD->mFormatFlags & kAudioFormatFlagIsNonMixable) != 0)
+				{
+					OutFormat.TypeName += TEXT(" (Non-Mixable)");
+				}
+			
+				return true;
+			}
+		}
+	}
+	return false;
+}
+
+int32 FAvfMediaCapturePlayer::GetNumTracks(EMediaTrackType TrackType) const
+{
+	SCOPED_AUTORELEASE_POOL;
+	
+	FScopeLock Lock(&CriticalSection);
+	if(MediaCaptureHelper != nil)
+	{
+		AVMediaType MediaType = [MediaCaptureHelper getCaptureDeviceMediaType];
+		if	(	(TrackType == EMediaTrackType::Video && MediaType == AVMediaTypeVideo) ||
+				(TrackType == EMediaTrackType::Audio && MediaType == AVMediaTypeAudio)
+			)
+		{
+			return 1;
+		}
+	}
+	return 0;
+}
+
+int32 FAvfMediaCapturePlayer::GetNumTrackFormats(EMediaTrackType TrackType, int32 TrackIndex) const
+{
+	SCOPED_AUTORELEASE_POOL;
+	
+	if(TrackIndex == 0)
+	{
+		FScopeLock Lock(&CriticalSection);
+		if(MediaCaptureHelper != nil)
+		{
+			AVMediaType MediaType = [MediaCaptureHelper getCaptureDeviceMediaType];
+			if	(	(TrackType == EMediaTrackType::Video && MediaType == AVMediaTypeVideo) ||
+					(TrackType == EMediaTrackType::Audio && MediaType == AVMediaTypeAudio)
+				)
+			{
+				return [MediaCaptureHelper getCaptureDeviceAvailableFormats].count;
+			}
+		}
+	}
+	
+	return 0;
+}
+
+int32 FAvfMediaCapturePlayer::GetSelectedTrack(EMediaTrackType TrackType) const
+{
+	SCOPED_AUTORELEASE_POOL;
+	
+	FScopeLock Lock(&CriticalSection);
+	if(MediaCaptureHelper != nil)
+	{
+		AVMediaType MediaType = [MediaCaptureHelper getCaptureDeviceMediaType];
+		if	(	(TrackType == EMediaTrackType::Video && MediaType == AVMediaTypeVideo) ||
+				(TrackType == EMediaTrackType::Audio && MediaType == AVMediaTypeAudio)
+			)
+		{
+			return 0;
+		}
+	}
+	return INDEX_NONE;
+}
+
+FText FAvfMediaCapturePlayer::GetTrackDisplayName(EMediaTrackType TrackType, int32 TrackIndex) const
+{
+	return FText();
+}
+
+int32 FAvfMediaCapturePlayer::GetTrackFormat(EMediaTrackType TrackType, int32 TrackIndex) const
+{
+	SCOPED_AUTORELEASE_POOL;
+	
+	if(TrackIndex == 0)
+	{
+		FScopeLock Lock(&CriticalSection);
+		if(MediaCaptureHelper != nil)
+		{
+			AVMediaType MediaType = [MediaCaptureHelper getCaptureDeviceMediaType];
+			if	(	(TrackType == EMediaTrackType::Video && MediaType == AVMediaTypeVideo) ||
+					(TrackType == EMediaTrackType::Audio && MediaType == AVMediaTypeAudio)
+				)
+			{
+				return [MediaCaptureHelper getCaptureDeviceActiveFormatIndex];
+			}
+		}
+	}
+	return INDEX_NONE;
+}
+
+FString FAvfMediaCapturePlayer::GetTrackLanguage(EMediaTrackType TrackType, int32 TrackIndex) const
+{
+	return FString();
+}
+
+FString FAvfMediaCapturePlayer::GetTrackName(EMediaTrackType TrackType, int32 TrackIndex) const
+{
+	return FString();
+}
+
+bool FAvfMediaCapturePlayer::GetVideoTrackFormat(int32 TrackIndex, int32 FormatIndex, FMediaVideoTrackFormat& OutFormat) const
+{
+	SCOPED_AUTORELEASE_POOL;
+	
+	FScopeLock Lock(&CriticalSection);
+	
+	if(TrackIndex == 0 && FormatIndex >= 0)
+	{
+		NSArray<AVCaptureDeviceFormat*>* deviceFormats = [MediaCaptureHelper getCaptureDeviceAvailableFormats];
+		if(FormatIndex < deviceFormats.count)
+		{
+			AVCaptureDeviceFormat* availableFormat = deviceFormats[FormatIndex];
+			if(CMFormatDescriptionGetMediaType(availableFormat.formatDescription) == kCMMediaType_Video)
+			{
+				float fMinRate = FLT_MAX;
+				float fMaxRate = FLT_MIN;
+				
+				NSArray<AVFrameRateRange*>* rateRanges = availableFormat.videoSupportedFrameRateRanges;
+				for(uint32 i = 0;i < rateRanges.count;++i)
+				{
+					AVFrameRateRange* range = rateRanges[i];
+					if(range.maxFrameRate > fMaxRate)
+					{
+						fMaxRate = range.maxFrameRate;
+					}
+					if(range.minFrameRate < fMinRate)
+					{
+						fMinRate = range.minFrameRate;
+					}
+				}
+				
+				// Missing rate info - set a sensible value
+				if(fMinRate > fMaxRate)
+				{
+					fMinRate = fMaxRate = 30.f;
+				}
+				
+				CMVideoDimensions Dimensions = CMVideoFormatDescriptionGetDimensions(availableFormat.formatDescription);
+				
+				OutFormat.Dim = FIntPoint(Dimensions.width, Dimensions.height);
+				OutFormat.FrameRate = fMaxRate;
+				OutFormat.FrameRates = TRange<float>(fMinRate, fMaxRate);
+				
+				FourCharCode subTypeCode = CMFormatDescriptionGetMediaSubType(availableFormat.formatDescription);
+				OutFormat.TypeName = FStringFromAppleFourCharWordCode(subTypeCode);
+				
+				return true;
+			}
+		}
+	}
+	return false;
+}
+
+bool FAvfMediaCapturePlayer::SelectTrack(EMediaTrackType TrackType, int32 TrackIndex)
+{
+	return TrackIndex == 0;
+}
+
+bool FAvfMediaCapturePlayer::SetTrackFormat(EMediaTrackType TrackType, int32 TrackIndex, int32 FormatIndex)
+{
+	SCOPED_AUTORELEASE_POOL;
+	
+	if(TrackIndex == 0)
+	{
+		FScopeLock Lock(&CriticalSection);
+		if(TrackType == EMediaTrackType::Video || TrackType == EMediaTrackType::Audio)
+		{
+			if([MediaCaptureHelper setCaptureDeviceActiveFormatIndex:FormatIndex])
+			{
+				return true;
+			}
+		}
+	}
+	return false;
+}
+
+// IMediaControls Interface
+
+bool FAvfMediaCapturePlayer::CanControl(EMediaControl Control) const
+{
+	return false;
+}
+
+FTimespan FAvfMediaCapturePlayer::GetDuration() const
+{
+	return FTimespan::MaxValue();
+}
+
+float FAvfMediaCapturePlayer::GetRate() const
+{
+	FScopeLock Lock(&CriticalSection);
+	return CurrentRate;
+}
+
+EMediaState FAvfMediaCapturePlayer::GetState() const
+{
+	SCOPED_AUTORELEASE_POOL;
+	
+	FScopeLock Lock(&CriticalSection);
+	if(MediaCaptureHelper != nil)
+	{
+		if(CurrentRate == 0.f && ![MediaCaptureHelper isCaptureRunning])
+		{
+			return EMediaState::Paused;
+		}
+		else
+		{
+			return EMediaState::Playing;
+		}
+	}
+	return EMediaState::Closed;
+}
+
+EMediaStatus FAvfMediaCapturePlayer::GetStatus() const
+{
+	return EMediaStatus::None;
+}
+
+TRangeSet<float> FAvfMediaCapturePlayer::GetSupportedRates(EMediaRateThinning Thinning) const
+{
+	TRangeSet<float> Result;
+
+	Result.Add(TRange<float>{0.0f});
+	Result.Add(TRange<float>{1.0f});
+
+	return Result;
+}
+
+FTimespan FAvfMediaCapturePlayer::GetTime() const
+{
+	FScopeLock Lock(&CriticalSection);
+	return CurrentTime;
+}
+
+bool FAvfMediaCapturePlayer::IsLooping() const
+{
+	return false;
+}
+
+bool FAvfMediaCapturePlayer::Seek(const FTimespan& Time)
+{
+	return false;
+}
+
+bool FAvfMediaCapturePlayer::SetLooping(bool Looping)
+{
+	return false;
+}
+
+bool FAvfMediaCapturePlayer::SetRate(float Rate)
+{
+	SCOPED_AUTORELEASE_POOL
+	
+	FScopeLock Lock(&CriticalSection);
+	if(Rate != CurrentRate)
+	{
+		if(Rate == 0.f && CurrentRate != 0.f)
+		{
+			[MediaCaptureHelper stopCaptureSession];
+		}
+		else if(Rate != 0.f && CurrentRate == 0.f)
+		{
+			[MediaCaptureHelper startCaptureSession];
+		}
+		
+		CurrentRate = Rate;
+	}
+	return true;
+}
+
+void FAvfMediaCapturePlayer::CaptureSystemNotification(NSNotification* Notification)
+{
+	NSString* NotificationName = Notification.name;
+	if(NotificationName == AVCaptureSessionRuntimeErrorNotification)
+	{
+		Close();
+	}
+	else if(NotificationName == AVCaptureSessionDidStartRunningNotification)
+	{
+		EventSink.ReceiveMediaEvent(EMediaEvent::PlaybackResumed);
+	}
+	else if(NotificationName == AVCaptureSessionDidStopRunningNotification)
+	{
+		EventSink.ReceiveMediaEvent(EMediaEvent::PlaybackSuspended);
+	}
+	else if(NotificationName == AVCaptureSessionWasInterruptedNotification)
+	{
+		EventSink.ReceiveMediaEvent(EMediaEvent::PlaybackSuspended);
+	}
+	else if(NotificationName == AVCaptureSessionInterruptionEndedNotification)
+	{
+		EventSink.ReceiveMediaEvent(EMediaEvent::PlaybackResumed);
+	}
+#if PLATFORM_MAC && WITH_EDITOR
+	else if(NotificationName == NSApplicationDidBecomeActiveNotification)
+	{
+		ThrottleDuration = 0.0;
+	}
+	else if(NotificationName == NSApplicationWillResignActiveNotification)
+	{
+		ThrottleDuration = 1.0;
+	}
+#endif
+}
+
+void FAvfMediaCapturePlayer::NewSampleBufferAvailable(CMSampleBufferRef SampleBuffer)
+{
+	if(!CMSampleBufferIsValid(SampleBuffer) || !CMSampleBufferDataIsReady(SampleBuffer))
+	{
+		return;
+	}
+	
+	FScopeLock Lock(&CriticalSection);
+	
+	if(CurrentRate == 0.f)
+	{
+		return;
+	}
+
+	// Allow each media type to compute the duration as CMSampleBuffer doesn't always report valid timings
+	CMFormatDescriptionRef FormatDescription = CMSampleBufferGetFormatDescription(SampleBuffer);
+	if(FormatDescription != NULL)
+	{
+		CMMediaType MediaType = CMFormatDescriptionGetMediaType(FormatDescription);
+		if(MediaType == kCMMediaType_Audio)
+		{
+			ProcessSampleBufferAudio(SampleBuffer);
+		}
+		else if(MediaType == kCMMediaType_Video)
+		{
+			ProcessSampleBufferVideo(SampleBuffer);
+		}
+	}
+}
+
+FTimespan FAvfMediaCapturePlayer::UpdateInternalTime(CMSampleBufferRef SampleBuffer, FTimespan const& ComputedBufferDuration)
+{
+	FTimespan SampleDuration(0);
+		
+	{
+		CMTime SampleBufferDuration = CMSampleBufferGetDuration(SampleBuffer);
+		if((SampleBufferDuration.flags & kCMTimeFlags_Valid) == kCMTimeFlags_Valid)
+		{
+			SampleDuration = FTimespan::FromSeconds(CMTimeGetSeconds(SampleBufferDuration));
+		}
+	}
+	
+	if(SampleDuration.IsZero())
+	{
+		SampleDuration = ComputedBufferDuration;
+	}
+	
+	if(CurrentTime == FTimespan::MinValue())
+	{
+		CurrentTime = FTimespan::Zero();
+	}
+	else if(CurrentTime >= FTimespan::Zero())
+	{
+		CurrentTime += SampleDuration;
+	}
+	
+	return SampleDuration;
+}
+
+void FAvfMediaCapturePlayer::ProcessSampleBufferVideo(CMSampleBufferRef SampleBuffer)
+{
+	check(SampleBuffer);
+	
+	FTimespan SampleDuration = UpdateInternalTime(SampleBuffer, FTimespan::FromSeconds(1.0 / 60.0));
+	
+#if PLATFORM_MAC && WITH_EDITOR
+	// Throttle frames pasted to the editor when not in focus
+	double PlatformTime = FPlatformTime::Seconds();
+	if(ThrottleDuration > 0.0 && (PlatformTime - LastConsumedTimeStamp) < ThrottleDuration)
+	{
+		return;
+	}
+	LastConsumedTimeStamp = PlatformTime;
+#endif
+
+	CVPixelBufferRef PixelBuffer = CMSampleBufferGetImageBuffer(SampleBuffer);
+	if(PixelBuffer)
+	{
+		check(CMSampleBufferGetNumSamples(SampleBuffer) == 1);
+		check(!CVPixelBufferIsPlanar(PixelBuffer));
+		
+		if (!MetalTextureCache)
+		{
+			id<MTLDevice> Device = (id<MTLDevice>)GDynamicRHI->RHIGetNativeDevice();
+			check(Device);
+			
+			CVReturn Return = CVMetalTextureCacheCreate(kCFAllocatorDefault, nullptr, Device, nullptr, &MetalTextureCache);
+			check(Return == kCVReturnSuccess);
+		}
+
+		int32 Width = CVPixelBufferGetWidth(PixelBuffer);
+		int32 Height = CVPixelBufferGetHeight(PixelBuffer);
+		
+		CVMetalTextureRef TextureRef = nullptr;
+		CVReturn Result = CVMetalTextureCacheCreateTextureFromImage(kCFAllocatorDefault, MetalTextureCache, PixelBuffer, nullptr, MTLPixelFormatBGRA8Unorm_sRGB, Width, Height, 0, &TextureRef);
+		check(Result == kCVReturnSuccess);
+		check(TextureRef);
+
+		const FRHITextureCreateDesc Desc =
+			FRHITextureCreateDesc::Create2D(TEXT("FAvfMediaCapturePlayer"), Width, Height, PF_B8G8R8A8)
+			.SetFlags(ETextureCreateFlags::SRGB | ETextureCreateFlags::Dynamic | ETextureCreateFlags::NoTiling | ETextureCreateFlags::ShaderResource)
+			.SetBulkData(new FAvfTexture2DResourceWrapper(TextureRef));
+
+		TRefCountPtr<FRHITexture2D> ShaderResource = RHICreateTexture(Desc);
+		CFRelease(TextureRef);
+		
+		FIntPoint Dim(Width, Height);
+		
+		TSharedRef<FAvfMediaTextureSample, ESPMode::ThreadSafe> VideoSample = VideoSamplePool.AcquireShared();
+		VideoSample->Initialize(ShaderResource, Dim, Dim, CurrentTime, SampleDuration);
+		
+		MediaSamples->AddVideo(VideoSample);
+	}
+}
+
+void FAvfMediaCapturePlayer::ProcessSampleBufferAudio(CMSampleBufferRef SampleBuffer)
+{
+	check(SampleBuffer);
+	
+	// Avoid using CMSampleBufferGetAudioBufferListWithRetainedBlockBuffer as that appears to copy data - don't want to incurr the extra copy before we even copy out into the UE structure
+	CMFormatDescriptionRef FormatDescription = CMSampleBufferGetFormatDescription(SampleBuffer);
+	if(FormatDescription != NULL && CMFormatDescriptionGetMediaType(FormatDescription) == kCMMediaType_Audio)
+	{
+		CMAudioFormatDescriptionRef AudioFormatDescription = (CMAudioFormatDescriptionRef)FormatDescription;
+		AudioStreamBasicDescription const* ASBD = CMAudioFormatDescriptionGetStreamBasicDescription(AudioFormatDescription);
+		check(ASBD);
+		
+		CMBlockBufferRef BlockBuffer = CMSampleBufferGetDataBuffer(SampleBuffer);
+		if(BlockBuffer != NULL)
+		{
+			const size_t BufferSize = CMBlockBufferGetDataLength(BlockBuffer);
+			const size_t BitDepth = ASBD->mBitsPerChannel;
+			const size_t NumChannels = ASBD->mChannelsPerFrame;
+			
+			// CMSampleBufferGetNumSamples() is per channel so this is really the frame count
+			const size_t NumFrames = CMSampleBufferGetNumSamples(SampleBuffer);
+			
+			// It should match a computed value - if it doesn't then use this computed value instead as Apple switch meaning of Frames and Samples depending on doc / API.
+			check(NumFrames == BufferSize / NumChannels / (BitDepth / 8));
+			
+			// Get valid engine format
+			EMediaAudioSampleFormat SampleFormat = EMediaAudioSampleFormat::Undefined;
+			if((ASBD->mFormatFlags & kLinearPCMFormatFlagIsFloat) != 0)
+			{
+				if(BitDepth == 32)		SampleFormat = EMediaAudioSampleFormat::Float;
+				else if(BitDepth == 64)	SampleFormat = EMediaAudioSampleFormat::Double;
+			}
+			else
+			{
+				if(BitDepth == 32)		SampleFormat = EMediaAudioSampleFormat::Int32;
+				else if(BitDepth == 24) SampleFormat = EMediaAudioSampleFormat::Undefined; // 24-bit int is not supported by UE
+				else if(BitDepth == 16)	SampleFormat = EMediaAudioSampleFormat::Int16;
+				else if(BitDepth == 8)	SampleFormat = EMediaAudioSampleFormat::Int8;
+			}
+
+			FTimespan SampleDuration = UpdateInternalTime(SampleBuffer, FTimespan::FromSeconds(((double)NumFrames) / ASBD->mSampleRate));
+			
+			if(SampleFormat != EMediaAudioSampleFormat::Undefined)
+			{
+				TSharedRef<FAvfMediaAudioSample, ESPMode::ThreadSafe> AudioSample = AudioSamplePool.AcquireShared();
+				if (AudioSample->Initialize(BufferSize,
+											NumFrames,
+											NumChannels,
+											ASBD->mSampleRate,
+											CurrentTime,
+											SampleDuration,
+											SampleFormat))
+				{
+					size_t CurrentOffset = 0;
+					uint8* DestData = AudioSample->GetMutableBuffer();
+
+					while(CurrentOffset < BufferSize)
+					{
+						char* SrcData = NULL;
+
+						size_t LengthAtOffset = 0;
+						size_t BlockTotalSize = 0;
+
+						if(CMBlockBufferGetDataPointer(BlockBuffer, CurrentOffset, &LengthAtOffset, &BlockTotalSize, &SrcData) != noErr)
+						{
+							break; // can't get data pointer force exit while loop
+						}
+
+						if(LengthAtOffset == 0)
+						{
+							break; // no data force exit while loop
+						}
+
+						checkf((CurrentOffset + LengthAtOffset) <= BufferSize, TEXT("Buffer overflow"));
+						checkf(BlockTotalSize == BufferSize, TEXT("Total CMBlockBuffer size does not match data pointer total sizes"));
+						checkf(CMBlockBufferIsRangeContiguous(BlockBuffer, CurrentOffset, LengthAtOffset), TEXT("Returned data pointer for CMBlockBuffer range is not contiguous"));
+
+						if(SrcData != NULL)
+						{
+							if((ASBD->mFormatFlags & kAudioFormatFlagIsNonInterleaved) == 0 || NumChannels == 1)
+							{
+								// Interleaved data OR single channel - direct copy
+								FMemory::Memcpy(DestData + CurrentOffset, SrcData, LengthAtOffset);
+							}
+							else
+							{
+								// Non Interleaved AND multiple channel - copy to interleaved buffer
+								size_t SrcNumFrames = LengthAtOffset / (BitDepth / 8);
+								size_t SrcFrameOffset = CurrentOffset / (BitDepth / 8);
+								
+								if(SampleFormat == EMediaAudioSampleFormat::Float) 			InterleaveCopy<float>(DestData, SrcData, NumChannels, NumFrames, SrcFrameOffset, SrcNumFrames);
+								else if(SampleFormat == EMediaAudioSampleFormat::Double) 	InterleaveCopy<double>(DestData, SrcData, NumChannels, NumFrames, SrcFrameOffset, SrcNumFrames);
+								else if(SampleFormat == EMediaAudioSampleFormat::Int32) 	InterleaveCopy<int32>(DestData, SrcData, NumChannels, NumFrames, SrcFrameOffset, SrcNumFrames);
+								else if(SampleFormat == EMediaAudioSampleFormat::Int16) 	InterleaveCopy<int16>(DestData, SrcData, NumChannels, NumFrames, SrcFrameOffset, SrcNumFrames);
+								else if(SampleFormat == EMediaAudioSampleFormat::Int8) 		InterleaveCopy<int8>(DestData, SrcData, NumChannels, NumFrames, SrcFrameOffset, SrcNumFrames);
+								else break;	// unsupported force exit while loop
+							}
+						}
+
+						CurrentOffset += LengthAtOffset;
+					}
+					
+					if(CurrentOffset < BufferSize)
+					{
+						// Error case - clear the output buffer to avoid unwanted noise
+						FMemory::Memset(DestData, 0, BufferSize);
+					}
+
+					MediaSamples->AddAudio(AudioSample);
+				}
+			}
+		}
+	}
+}
diff --git a/Engine/Plugins/Media/AvfMedia/Source/AvfMediaCapture/Private/Player/AvfMediaCapturePlayer.h b/Engine/Plugins/Media/AvfMedia/Source/AvfMediaCapture/Private/Player/AvfMediaCapturePlayer.h
new file mode 100644
--- /dev/null
+++ b/Engine/Plugins/Media/AvfMedia/Source/AvfMediaCapture/Private/Player/AvfMediaCapturePlayer.h
@@ -0,0 +1,121 @@
+// Copyright Epic Games, Inc. All Rights Reserved.
+
+#pragma once
+
+#include "IMediaCache.h"
+#include "IMediaControls.h"
+#include "IMediaPlayer.h"
+#include "IMediaTracks.h"
+#include "IMediaView.h"
+#include "IMediaTextureSample.h"
+#include "IMediaAudioSample.h"
+#include "IMediaEventSink.h"
+#include "Misc/Timespan.h"
+
+#include "AvfMediaCaptureHelper.h"
+#include "AvfMediaAudioSample.h"
+#include "AvfMediaTextureSample.h"
+
+class FMediaSamples;
+class FAvfMediaAudioSamplePool;
+class FAvfMediaTextureSamplePool;
+
+class FAvfMediaCapturePlayer
+	: public IMediaPlayer
+	, public IMediaTracks
+	, protected IMediaCache
+	, protected IMediaControls
+	, protected IMediaView
+	, public TSharedFromThis<FAvfMediaCapturePlayer>
+{
+public:
+	
+	FAvfMediaCapturePlayer(IMediaEventSink& InEventSink);
+	virtual ~FAvfMediaCapturePlayer();
+
+public:
+	// IMediaPlayer Interface
+	virtual void Close() override;
+	virtual IMediaCache& GetCache() override;
+	virtual IMediaControls& GetControls() override;
+	virtual FString GetInfo() const override;
+	virtual FGuid GetPlayerPluginGUID() const override;
+	virtual IMediaSamples& GetSamples() override;
+	virtual FString GetStats() const override;
+	virtual IMediaTracks& GetTracks() override;
+	virtual FString GetUrl() const override;
+	virtual IMediaView& GetView() override;
+	virtual bool Open(const FString& Url, const IMediaOptions* Options) override;
+	virtual bool Open(const TSharedRef<FArchive, ESPMode::ThreadSafe>& Archive, const FString& OriginalUrl, const IMediaOptions* Options) override;
+	virtual FText GetMediaName() const override;
+	
+	// IMediaTracks Interface
+	virtual bool GetAudioTrackFormat(int32 TrackIndex, int32 FormatIndex, FMediaAudioTrackFormat& OutFormat) const override;
+	virtual int32 GetNumTracks(EMediaTrackType TrackType) const override;
+	virtual int32 GetNumTrackFormats(EMediaTrackType TrackType, int32 TrackIndex) const override;
+	virtual int32 GetSelectedTrack(EMediaTrackType TrackType) const override;
+	virtual FText GetTrackDisplayName(EMediaTrackType TrackType, int32 TrackIndex) const override;
+	virtual int32 GetTrackFormat(EMediaTrackType TrackType, int32 TrackIndex) const override;
+	virtual FString GetTrackLanguage(EMediaTrackType TrackType, int32 TrackIndex) const override;
+	virtual FString GetTrackName(EMediaTrackType TrackType, int32 TrackIndex) const override;
+	virtual bool GetVideoTrackFormat(int32 TrackIndex, int32 FormatIndex, FMediaVideoTrackFormat& OutFormat) const override;
+	virtual bool SelectTrack(EMediaTrackType TrackType, int32 TrackIndex) override;
+	virtual bool SetTrackFormat(EMediaTrackType TrackType, int32 TrackIndex, int32 FormatIndex) override;
+
+protected:
+	
+	// IMediaControls Interface
+	virtual bool CanControl(EMediaControl Control) const override;
+	virtual FTimespan GetDuration() const override;
+	virtual float GetRate() const override;
+	virtual EMediaState GetState() const override;
+	virtual EMediaStatus GetStatus() const override;
+	virtual TRangeSet<float> GetSupportedRates(EMediaRateThinning Thinning) const override;
+	virtual FTimespan GetTime() const override;
+	virtual bool IsLooping() const override;
+	virtual bool Seek(const FTimespan& Time) override;
+	virtual bool SetLooping(bool Looping) override;
+	virtual bool SetRate(float Rate) override;
+	
+private:
+
+	void CreateCaptureSession(NSString* deviceIDString);
+	void CaptureSystemNotification(NSNotification* Notification);
+	void HandleAuthStatusError(EAvfMediaCaptureAuthStatus AuthStatus, AVMediaType MediaType);
+
+	FTimespan UpdateInternalTime(CMSampleBufferRef SampleBuffer, FTimespan const& ComputedBufferDuration);
+	
+	void NewSampleBufferAvailable(CMSampleBufferRef SampleBuffer);
+	void ProcessSampleBufferVideo(CMSampleBufferRef SampleBuffer);
+	void ProcessSampleBufferAudio(CMSampleBufferRef SampleBuffer);
+
+private:
+
+	/** The media event handler. */
+	IMediaEventSink& EventSink;
+
+	/** Media sample collection. */
+	TSharedPtr<FMediaSamples, ESPMode::ThreadSafe> MediaSamples;
+	
+	/** Media capture AV Obj-C Helper */
+	AvfMediaCaptureHelper* MediaCaptureHelper;
+	
+	/** Reuse Sample Pool Objects */
+	FAvfMediaAudioSamplePool	AudioSamplePool;
+	FAvfMediaTextureSamplePool	VideoSamplePool;
+	
+	/** Internal CoreVideo Metal Texture cache to handle AVFoundation resource pooling and optimisations */
+	CVMetalTextureCacheRef		MetalTextureCache;
+	
+	/** Media Playback Info and Control data */
+	float 		CurrentRate;
+	FTimespan 	CurrentTime;
+	FString 	URL;
+
+#if PLATFORM_MAC && WITH_EDITOR
+	double		ThrottleDuration;
+	double		LastConsumedTimeStamp;
+#endif
+
+	mutable FCriticalSection CriticalSection;
+};
