Fix for NoPCH macOS Game Build

Horde: 212304

diff --git a/Engine/Plugins/Media/AvfMedia/Source/AvfMedia/Public/AvfMediaTextureSample.h b/Engine/Plugins/Media/AvfMedia/Source/AvfMedia/Public/AvfMediaTextureSample.h
--- a/Engine/Plugins/Media/AvfMedia/Source/AvfMedia/Public/AvfMediaTextureSample.h
+++ b/Engine/Plugins/Media/AvfMedia/Source/AvfMedia/Public/AvfMediaTextureSample.h
@@ -8,6 +8,7 @@
 #include "Math/IntPoint.h"
 #include "Math/Range.h"
 #include "Misc/Timespan.h"
+#include "Containers/ResourceArray.h"
 
 #if WITH_ENGINE
 	#include "RHI.h"
