Non unity mac compile fix.

Horde: 212304
#rnx

diff --git a/Engine/Plugins/Media/AvfMedia/Source/AvfMediaCapture/Private/Player/AvfMediaCaptureHelper.cpp b/Engine/Plugins/Media/AvfMedia/Source/AvfMediaCapture/Private/Player/AvfMediaCaptureHelper.cpp
--- a/Engine/Plugins/Media/AvfMedia/Source/AvfMediaCapture/Private/Player/AvfMediaCaptureHelper.cpp
+++ b/Engine/Plugins/Media/AvfMedia/Source/AvfMediaCapture/Private/Player/AvfMediaCaptureHelper.cpp
@@ -2,6 +2,8 @@
 
 #import <Foundation/Foundation.h>
 #import <AVFoundation/AVFoundation.h>
+
+#include "CoreTypes.h"
 #include "AvfMediaCaptureHelper.h"
 
 // Steps to try if permissions are not working correctly in macOS
