macOS header include order fix.

Horde: 212774
#rnx

diff --git a/Engine/Plugins/Media/AvfMedia/Source/AvfMediaCapture/Private/Player/AvfMediaCaptureHelper.cpp b/Engine/Plugins/Media/AvfMedia/Source/AvfMediaCapture/Private/Player/AvfMediaCaptureHelper.cpp
--- a/Engine/Plugins/Media/AvfMedia/Source/AvfMediaCapture/Private/Player/AvfMediaCaptureHelper.cpp
+++ b/Engine/Plugins/Media/AvfMedia/Source/AvfMediaCapture/Private/Player/AvfMediaCaptureHelper.cpp
@@ -3,7 +3,6 @@
 #import <Foundation/Foundation.h>
 #import <AVFoundation/AVFoundation.h>
 
-#include "CoreTypes.h"
 #include "AvfMediaCaptureHelper.h"
 
 // Steps to try if permissions are not working correctly in macOS
diff --git a/Engine/Plugins/Media/AvfMedia/Source/AvfMediaCapture/Private/Player/AvfMediaCaptureHelper.h b/Engine/Plugins/Media/AvfMedia/Source/AvfMediaCapture/Private/Player/AvfMediaCaptureHelper.h
--- a/Engine/Plugins/Media/AvfMedia/Source/AvfMediaCapture/Private/Player/AvfMediaCaptureHelper.h
+++ b/Engine/Plugins/Media/AvfMedia/Source/AvfMediaCapture/Private/Player/AvfMediaCaptureHelper.h
@@ -3,6 +3,7 @@
 #pragma once
 
 #import <AVFoundation/AVFoundation.h>
+#include "CoreTypes.h"
 
 enum class EAvfMediaCaptureAuthStatus : uint32
 {
