macOS non unity editor compile issue for app kit notifications.

Horde: 212304
#rnx

diff --git a/Engine/Plugins/Media/AvfMedia/Source/AvfMediaCapture/Private/Player/AvfMediaCaptureHelper.cpp b/Engine/Plugins/Media/AvfMedia/Source/AvfMediaCapture/Private/Player/AvfMediaCaptureHelper.cpp
--- a/Engine/Plugins/Media/AvfMedia/Source/AvfMediaCapture/Private/Player/AvfMediaCaptureHelper.cpp
+++ b/Engine/Plugins/Media/AvfMedia/Source/AvfMediaCapture/Private/Player/AvfMediaCaptureHelper.cpp
@@ -1,8 +1,5 @@
 // Copyright Epic Games, Inc. All Rights Reserved.
 
-#import <Foundation/Foundation.h>
-#import <AVFoundation/AVFoundation.h>
-
 #include "AvfMediaCaptureHelper.h"
 
 // Steps to try if permissions are not working correctly in macOS
diff --git a/Engine/Plugins/Media/AvfMedia/Source/AvfMediaCapture/Private/Player/AvfMediaCaptureHelper.h b/Engine/Plugins/Media/AvfMedia/Source/AvfMediaCapture/Private/Player/AvfMediaCaptureHelper.h
--- a/Engine/Plugins/Media/AvfMedia/Source/AvfMediaCapture/Private/Player/AvfMediaCaptureHelper.h
+++ b/Engine/Plugins/Media/AvfMedia/Source/AvfMediaCapture/Private/Player/AvfMediaCaptureHelper.h
@@ -2,9 +2,15 @@
 
 #pragma once
 
-#import <AVFoundation/AVFoundation.h>
 #include "CoreTypes.h"
 
+#import <Foundation/Foundation.h>
+#import <AVFoundation/AVFoundation.h>
+
+#if PLATFORM_MAC
+	#import <Appkit/Appkit.h>
+#endif
+
 enum class EAvfMediaCaptureAuthStatus : uint32
 {
 	// Mapped Apple AVFoundation values
