NSLOCTEXT conflict fix in AvfMediaCapture

#rnx

diff --git a/Engine/Plugins/Media/AvfMedia/Source/AvfMediaCapture/Private/Player/AvfMediaCapturePlayer.cpp b/Engine/Plugins/Media/AvfMedia/Source/AvfMediaCapture/Private/Player/AvfMediaCapturePlayer.cpp
--- a/Engine/Plugins/Media/AvfMedia/Source/AvfMediaCapture/Private/Player/AvfMediaCapturePlayer.cpp
+++ b/Engine/Plugins/Media/AvfMedia/Source/AvfMediaCapture/Private/Player/AvfMediaCapturePlayer.cpp
@@ -153,22 +153,22 @@ void FAvfMediaCapturePlayer::HandleAuthStatusError(EAvfMediaCaptureAuthStatus Au
 		}
 		case EAvfMediaCaptureAuthStatus::MissingInfoPListEntry:
 		{
-			ErrorString = NSLOCTEXT("AvfMediaCapture", "CaptureDeviceFail", "Cannot start capture: Missing Info.plist entry: ").ToString() + (MediaType == AVMediaTypeVideo ? TEXT("\"NSCameraUsageDescription\"") : TEXT("\"NSMicrophoneUsageDescription\""));
+			ErrorString = NSLOCTEXT("AvfMediaCapture", "AvfMediaCapture_DeviceFail_MissingInfoPListEntry", "Cannot start capture: Missing Info.plist entry: ").ToString() + (MediaType == AVMediaTypeVideo ? TEXT("\"NSCameraUsageDescription\"") : TEXT("\"NSMicrophoneUsageDescription\""));
 			break;
 		}
 		case EAvfMediaCaptureAuthStatus::Restricted:
 		{
-			ErrorString = NSLOCTEXT("AvfMediaCapture", "CaptureDeviceFail", "Cannot start capture: Restricted").ToString();
+			ErrorString = NSLOCTEXT("AvfMediaCapture", "AvfMediaCapture_DeviceFail_Restricted", "Cannot start capture: Restricted").ToString();
 			break;
 		}
 		case EAvfMediaCaptureAuthStatus::Denied:
 		{
-			ErrorString = NSLOCTEXT("AvfMediaCapture", "CaptureDeviceFail", "Cannot start capture: Denied").ToString();
+			ErrorString = NSLOCTEXT("AvfMediaCapture", "AvfMediaCapture_DeviceFail_Denied", "Cannot start capture: Denied").ToString();
 			break;
 		}
 		default:
 		{
-			ErrorString = NSLOCTEXT("AvfMediaCapture", "CaptureDeviceFail", "Cannot start capture: Unknown reason").ToString();
+			ErrorString = NSLOCTEXT("AvfMediaCapture", "AvfMediaCapture_DeviceFail_Unknown", "Cannot start capture: Unknown reason").ToString();
 			break;
 		}
 	}
@@ -179,7 +179,7 @@ void FAvfMediaCapturePlayer::HandleAuthStatusError(EAvfMediaCaptureAuthStatus Au
 
 		ErrorString += FString(TEXT(" (")) + URL + FString(TEXT(")"));
 		FText Msg = FText::FromString(ErrorString);
-		FText Title = NSLOCTEXT("AvfMediaCapture", "CaptureDeviceFail", "Capture Device Start Failed");
+		FText Title = NSLOCTEXT("AvfMediaCapture", "AvfMediaCapture_DeviceFail", "Capture Device Start Failed");
 		FMessageDialog::Open(EAppMsgType::Ok, Msg, &Title);
 	}
 }
