Fixing AvfMedia include paths.

diff --git a/Engine/Plugins/Media/AvfMedia/Source/AvfMedia/Private/AvfMediaModule.cpp b/Engine/Plugins/Media/AvfMedia/Source/AvfMedia/Private/AvfMediaModule.cpp
--- a/Engine/Plugins/Media/AvfMedia/Source/AvfMedia/Private/AvfMediaModule.cpp
+++ b/Engine/Plugins/Media/AvfMedia/Source/AvfMedia/Private/AvfMediaModule.cpp
@@ -4,7 +4,7 @@
 
 #include "Modules/ModuleManager.h"
 
-#include "AvfMediaPlayer.h"
+#include "Player/AvfMediaPlayer.h"
 #include "IAvfMediaModule.h"
 
 
diff --git a/Engine/Plugins/Media/AvfMedia/Source/AvfMedia/Private/Player/AvfMediaPlayer.cpp b/Engine/Plugins/Media/AvfMedia/Source/AvfMedia/Private/Player/AvfMediaPlayer.cpp
--- a/Engine/Plugins/Media/AvfMedia/Source/AvfMedia/Private/Player/AvfMediaPlayer.cpp
+++ b/Engine/Plugins/Media/AvfMedia/Source/AvfMedia/Private/Player/AvfMediaPlayer.cpp
@@ -19,7 +19,7 @@
 #include "HAL/FileManager.h"
 
 #include "AvfMediaTracks.h"
-#include "AvfMediaUtils.h"
+#include "Shared/AvfMediaUtils.h"
 #include "IMediaAudioSample.h"
 #include "HAL/PlatformFileManager.h"
 #include "Async/Async.h"
diff --git a/Engine/Plugins/Media/AvfMedia/Source/AvfMedia/Private/Player/AvfMediaTracks.cpp b/Engine/Plugins/Media/AvfMedia/Source/AvfMedia/Private/Player/AvfMediaTracks.cpp
--- a/Engine/Plugins/Media/AvfMedia/Source/AvfMedia/Private/Player/AvfMediaTracks.cpp
+++ b/Engine/Plugins/Media/AvfMedia/Source/AvfMedia/Private/Player/AvfMediaTracks.cpp
@@ -14,7 +14,7 @@
 #include "AvfMediaAudioSample.h"
 #include "AvfMediaOverlaySample.h"
 #include "AvfMediaVideoSampler.h"
-#include "AvfMediaUtils.h"
+#include "Shared/AvfMediaUtils.h"
 
 #import <AudioToolbox/AudioToolbox.h>
 
