1. Fix a bug that the loadValuesAsynchronouslyForKeys could be invoked after PlayerItem is released in AvfMediaPlayer
2. Fix a bug that the seekToTime could be invoked after MediaPlayer is released in AvfMediaPlayer
3. Fix a bug that the loading task is still running after the MediaResourceLoader is released in AvfMediaPlayer

diff --git a/Engine/Plugins/Media/AvfMedia/Source/AvfMedia/Private/Player/AvfMediaPlayer.cpp b/Engine/Plugins/Media/AvfMedia/Source/AvfMedia/Private/Player/AvfMediaPlayer.cpp
--- a/Engine/Plugins/Media/AvfMedia/Source/AvfMedia/Private/Player/AvfMediaPlayer.cpp
+++ b/Engine/Plugins/Media/AvfMedia/Source/AvfMedia/Private/Player/AvfMediaPlayer.cpp
@@ -107,6 +107,7 @@
 	@public TSharedPtr<FArchive, ESPMode::ThreadSafe> FileAReader;
 	@public FCriticalSection CriticalSection;
 	@public bool bInitialized;
+	@public TFuture<void> AsyncLoadingTask;
 }
 @end
 
@@ -144,7 +145,7 @@
 	[loadingRequest retain];
 	
 	// Allow this function to return quickly so the resource loader knows the data is probabky coming and doesn't error
-	Async(EAsyncExecution::ThreadPool, [self, loadingRequest]()
+	AsyncLoadingTask = Async(EAsyncExecution::ThreadPool, [self, loadingRequest]()
 	{
 		FScopeLock ScopeLock(&CriticalSection);
 		
@@ -522,6 +523,11 @@ void FAvfMediaPlayer::Close()
 			[PlayerItem removeObserver:MediaHelper forKeyPath:@"status"];
 		}
 
+		// Override the completionHandler in case it hasn't been invoked before releasing PlayerItem.
+		[[PlayerItem asset] loadValuesAsynchronouslyForKeys:@[@"tracks"] completionHandler:^
+		{
+		}];
+
 		[PlayerItem release];
 		PlayerItem = nil;
 	}
@@ -536,12 +542,16 @@ void FAvfMediaPlayer::Close()
 	{
 		// If we don't remove the current player item then the retain count is > 1 for the MediaPlayer then on it's release then the MetalPlayer stays around forever
 		[MediaPlayer replaceCurrentItemWithPlayerItem:nil];
+		// Cancel the pending seeks before releasing the MediaPlayer
+		[MediaPlayer cancelPendingSeeks];
 		[MediaPlayer release];
 		MediaPlayer = nil;
 	}
 	
 	if(MediaResourceLoader != nil)
 	{
+		// Wait the loading task to finish before releasing the loader.
+		MediaResourceLoader->AsyncLoadingTask.Wait();
 		[MediaResourceLoader release];
 		MediaResourceLoader = nil;
 	}
