Fix a bug that introduced in CL#26144213

diff --git a/Engine/Plugins/Media/AvfMedia/Source/AvfMedia/Private/Player/AvfMediaPlayer.cpp b/Engine/Plugins/Media/AvfMedia/Source/AvfMedia/Private/Player/AvfMediaPlayer.cpp
--- a/Engine/Plugins/Media/AvfMedia/Source/AvfMedia/Private/Player/AvfMediaPlayer.cpp
+++ b/Engine/Plugins/Media/AvfMedia/Source/AvfMedia/Private/Player/AvfMediaPlayer.cpp
@@ -523,8 +523,11 @@ void FAvfMediaPlayer::Close()
 			[PlayerItem removeObserver:MediaHelper forKeyPath:@"status"];
 		}
 
+		// Cancel the pending seeks before releasing the MediaPlayer
+		[PlayerItem cancelPendingSeeks];
+		[PlayerItem.asset cancelLoading];
 		// Override the completionHandler in case it hasn't been invoked before releasing PlayerItem.
-		[[PlayerItem asset] loadValuesAsynchronouslyForKeys:@[@"tracks"] completionHandler:^
+		[PlayerItem.asset loadValuesAsynchronouslyForKeys:@[@"tracks"] completionHandler:^
 		{
 		}];
 
@@ -542,8 +545,6 @@ void FAvfMediaPlayer::Close()
 	{
 		// If we don't remove the current player item then the retain count is > 1 for the MediaPlayer then on it's release then the MetalPlayer stays around forever
 		[MediaPlayer replaceCurrentItemWithPlayerItem:nil];
-		// Cancel the pending seeks before releasing the MediaPlayer
-		[MediaPlayer cancelPendingSeeks];
 		[MediaPlayer release];
 		MediaPlayer = nil;
 	}
