- Fixed the AvfMediaCapture IOS17 changes compile on Mac

diff --git a/Engine/Plugins/Media/AvfMedia/Source/AvfMediaCapture/Private/AvfMediaCaptureModule.cpp b/Engine/Plugins/Media/AvfMedia/Source/AvfMediaCapture/Private/AvfMediaCaptureModule.cpp
--- a/Engine/Plugins/Media/AvfMedia/Source/AvfMediaCapture/Private/AvfMediaCaptureModule.cpp
+++ b/Engine/Plugins/Media/AvfMedia/Source/AvfMediaCapture/Private/AvfMediaCaptureModule.cpp
@@ -29,6 +29,14 @@ DEFINE_LOG_CATEGORY(LogAvfMediaCapture);
 #define LOCTEXT_NAMESPACE "FAvfMediaCaptureFactoryModule"
 
 
+// new API doesn't compile on old IOS sdk
+#if (PLATFORM_IOS && (defined(__IPHONE_17_0) && __IPHONE_OS_VERSION_MAX_ALLOWED >= __IPHONE_17_0))
+	#define USE_NEW_MICROPHONE_API 1
+#else
+	#define USE_NEW_MICROPHONE_API 0
+#endif
+
+
 /**
  * Implements the AvfMediaCapture module.
  */
@@ -44,7 +52,7 @@ public:
 
         NSArray* DeviceTypes = nil;
         
-    #if (defined(__IPHONE_17_0) && __IPHONE_OS_VERSION_MAX_ALLOWED >= __IPHONE_17_0)
+    #if USE_NEW_MICROPHONE_API
         DeviceTypes = @[AVCaptureDeviceTypeBuiltInWideAngleCamera, AVCaptureDeviceTypeMicrophone];
     #else
         DeviceTypes = @[AVCaptureDeviceTypeBuiltInWideAngleCamera, AVCaptureDeviceTypeBuiltInMicrophone];
@@ -57,7 +65,7 @@ public:
 			for(uint32 i = 0;i < Devices.count;++i)
 			{
 				AVCaptureDevice* AvailableDevice = Devices[i];
-            #if (defined(__IPHONE_17_0) && __IPHONE_OS_VERSION_MAX_ALLOWED >= __IPHONE_17_0)
+            #if USE_NEW_MICROPHONE_API
                 if(TargetDeviceType == EMediaCaptureDeviceType::Audio && AvailableDevice.deviceType == AVCaptureDeviceTypeMicrophone)
             #else
                 if(TargetDeviceType == EMediaCaptureDeviceType::Audio && AvailableDevice.deviceType == AVCaptureDeviceTypeBuiltInMicrophone)
@@ -72,7 +80,7 @@ public:
 					
 					OutDeviceInfos.Add(MoveTemp(DeviceInfo));
 				}
-            #if (defined(__IPHONE_17_0) && __IPHONE_OS_VERSION_MAX_ALLOWED >= __IPHONE_17_0)
+            #if USE_NEW_MICROPHONE_API
                 else if(TargetDeviceType == EMediaCaptureDeviceType::Video && AvailableDevice.deviceType != AVCaptureDeviceTypeMicrophone)
             #else
                 else if(TargetDeviceType == EMediaCaptureDeviceType::Video && AvailableDevice.deviceType != AVCaptureDeviceTypeBuiltInMicrophone)
diff --git a/Engine/Plugins/Media/AvfMedia/Source/AvfMediaCapture/Private/Player/AvfMediaCaptureHelper.cpp b/Engine/Plugins/Media/AvfMedia/Source/AvfMediaCapture/Private/Player/AvfMediaCaptureHelper.cpp
--- a/Engine/Plugins/Media/AvfMedia/Source/AvfMediaCapture/Private/Player/AvfMediaCaptureHelper.cpp
+++ b/Engine/Plugins/Media/AvfMedia/Source/AvfMediaCapture/Private/Player/AvfMediaCaptureHelper.cpp
@@ -14,6 +14,14 @@
 // 3) delete permissions database then reboot mac:
 //  ~/Library/Application\\ Support/com.apple.TCC
 
+// new API doesn't compile on old IOS sdk
+#if (PLATFORM_IOS && (defined(__IPHONE_17_0) && __IPHONE_OS_VERSION_MAX_ALLOWED >= __IPHONE_17_0))
+	#define USE_NEW_MICROPHONE_API 1
+#else
+	#define USE_NEW_MICROPHONE_API 0
+#endif
+
+
 @interface AvfMediaCaptureHelper()<AVCaptureVideoDataOutputSampleBufferDelegate, AVCaptureAudioDataOutputSampleBufferDelegate>
 
 @property (nonatomic,readwrite,assign) AVCaptureSession* 		captureSession;
@@ -172,7 +180,7 @@
 	AVCaptureDevice* foundDevice = nil;
     NSArray* deviceTypes = nil;
     
-#if (defined(__IPHONE_17_0) && __IPHONE_OS_VERSION_MAX_ALLOWED >= __IPHONE_17_0)
+#if USE_NEW_MICROPHONE_API
     deviceTypes = @[AVCaptureDeviceTypeBuiltInWideAngleCamera, AVCaptureDeviceTypeMicrophone];
 #else
     deviceTypes = @[AVCaptureDeviceTypeBuiltInWideAngleCamera, AVCaptureDeviceTypeBuiltInMicrophone];
@@ -234,7 +242,7 @@
 		if(self.captureSession.inputs.count > 0)
 		{
             // Video or Audio Device - need the correct output
-#if (defined(__IPHONE_17_0) && __IPHONE_OS_VERSION_MAX_ALLOWED >= __IPHONE_17_0)
+#if USE_NEW_MICROPHONE_API
             if(self.captureDevice.deviceType != AVCaptureDeviceTypeMicrophone)
 #else
             if(self.captureDevice.deviceType != AVCaptureDeviceTypeBuiltInMicrophone)
@@ -353,7 +361,7 @@
 	AVMediaType type = nil;
 	if(self.captureDevice != nil)
 	{
-#if (defined(__IPHONE_17_0) && __IPHONE_OS_VERSION_MAX_ALLOWED >= __IPHONE_17_0)
+#if USE_NEW_MICROPHONE_API
 		if(self.captureDevice.deviceType != AVCaptureDeviceTypeMicrophone)
 #else
         if(self.captureDevice.deviceType != AVCaptureDeviceTypeBuiltInMicrophone)
