Prevent crash on Mac when capturing video from a connected camera device via a Media Player asset by enqueuing the call to RHICreateTexture on the rendering thread.

Also fixes the 'No video track selected' message being displayed within the media player asset viewport when video capture is started using an AvfMediaCapturePlayer by emitting a 'TracksChanged' event after setting up the capture session.

diff --git a/Engine/Plugins/Media/AvfMedia/Source/AvfMediaCapture/Private/Player/AvfMediaCapturePlayer.cpp b/Engine/Plugins/Media/AvfMedia/Source/AvfMediaCapture/Private/Player/AvfMediaCapturePlayer.cpp
--- a/Engine/Plugins/Media/AvfMedia/Source/AvfMediaCapture/Private/Player/AvfMediaCapturePlayer.cpp
+++ b/Engine/Plugins/Media/AvfMedia/Source/AvfMediaCapture/Private/Player/AvfMediaCapturePlayer.cpp
@@ -206,6 +206,7 @@ void FAvfMediaCapturePlayer::CreateCaptureSession(NSString* deviceIDString)
 		if(bResult)
 		{
 			EventSink.ReceiveMediaEvent(EMediaEvent::MediaOpened);
+			EventSink.ReceiveMediaEvent(EMediaEvent::TracksChanged);
 		}
 		else
 		{
@@ -284,6 +285,11 @@ bool FAvfMediaCapturePlayer::Open(const TSharedRef<FArchive, ESPMode::ThreadSafe
 	return false;
 }
 
+bool FAvfMediaCapturePlayer::GetPlayerFeatureFlag(EFeatureFlag Flag) const
+{
+	return Flag == EFeatureFlag::PlayerSelectsDefaultTracks ? true : false;
+}
+
 // IMediaTracks Interface
 
 bool FAvfMediaCapturePlayer::GetAudioTrackFormat(int32 TrackIndex, int32 FormatIndex, FMediaAudioTrackFormat& OutFormat) const
@@ -718,20 +724,27 @@ void FAvfMediaCapturePlayer::ProcessSampleBufferVideo(CMSampleBufferRef SampleBu
 		check(Result == kCVReturnSuccess);
 		check(TextureRef);
 
-		const FRHITextureCreateDesc Desc =
+		FRHITextureCreateDesc Desc =
 			FRHITextureCreateDesc::Create2D(TEXT("FAvfMediaCapturePlayer"), Width, Height, PF_B8G8R8A8)
 			.SetFlags(ETextureCreateFlags::SRGB | ETextureCreateFlags::Dynamic | ETextureCreateFlags::NoTiling | ETextureCreateFlags::ShaderResource)
 			.SetBulkData(new FAvfTexture2DResourceWrapper(TextureRef));
 
-		TRefCountPtr<FRHITexture> ShaderResource = RHICreateTexture(Desc);
 		CFRelease(TextureRef);
 		
 		FIntPoint Dim(Width, Height);
 		
-		TSharedRef<FAvfMediaTextureSample, ESPMode::ThreadSafe> VideoSample = VideoSamplePool.AcquireShared();
-		VideoSample->Initialize(ShaderResource, Dim, Dim, CurrentTime, SampleDuration);
-		
-		MediaSamples->AddVideo(VideoSample);
+		ENQUEUE_RENDER_COMMAND(FAvfMediaCapturePlayer_ProcessSampleBufferVideo_CreateTexture)(
+            [WeakPlayer = this->AsWeak(), InDesc = MoveTemp(Desc), Dim, SampleDuration](FRHICommandListImmediate& RHICmdList)
+            {
+            	if (const TSharedPtr<FAvfMediaCapturePlayer> Player = WeakPlayer.Pin())
+            	{
+            		TRefCountPtr<FRHITexture> ShaderResource = RHICreateTexture(InDesc);
+					TSharedRef<FAvfMediaTextureSample, ESPMode::ThreadSafe> VideoSample = Player->VideoSamplePool.AcquireShared();
+					VideoSample->Initialize(ShaderResource, Dim, Dim, Player->CurrentTime, SampleDuration);
+            
+					Player->MediaSamples->AddVideo(VideoSample);
+            	}
+        });
 	}
 }
 
diff --git a/Engine/Plugins/Media/AvfMedia/Source/AvfMediaCapture/Private/Player/AvfMediaCapturePlayer.h b/Engine/Plugins/Media/AvfMedia/Source/AvfMediaCapture/Private/Player/AvfMediaCapturePlayer.h
--- a/Engine/Plugins/Media/AvfMedia/Source/AvfMediaCapture/Private/Player/AvfMediaCapturePlayer.h
+++ b/Engine/Plugins/Media/AvfMedia/Source/AvfMediaCapture/Private/Player/AvfMediaCapturePlayer.h
@@ -48,6 +48,7 @@ public:
 	virtual bool Open(const FString& Url, const IMediaOptions* Options) override;
 	virtual bool Open(const TSharedRef<FArchive, ESPMode::ThreadSafe>& Archive, const FString& OriginalUrl, const IMediaOptions* Options) override;
 	virtual FText GetMediaName() const override;
+	virtual bool GetPlayerFeatureFlag(EFeatureFlag Flag) const override;
 	
 	// IMediaTracks Interface
 	virtual bool GetAudioTrackFormat(int32 TrackIndex, int32 FormatIndex, FMediaAudioTrackFormat& OutFormat) const override;
