Extend capture device support in AvfMediaCapturePlayer.

On iOS:
- Adds built in ultra wide camera
- Adds built in telephoto camera

On iOS/iPadOS 17 and macOS
- Adds external camera devices such as USB webcams

In addition this reworks how we determine the availability of AVCaptureDeviceType values only available on iOS or macOS post 17.0 and 14.0 respectively.

A neat side-effect of this change is that allows Mac users to use their iPhone cameras over the local network.

diff --git a/Engine/Plugins/Media/AvfMedia/Source/AvfMediaCapture/Private/AvfMediaCaptureModule.cpp b/Engine/Plugins/Media/AvfMedia/Source/AvfMediaCapture/Private/AvfMediaCaptureModule.cpp
--- a/Engine/Plugins/Media/AvfMedia/Source/AvfMediaCapture/Private/AvfMediaCaptureModule.cpp
+++ b/Engine/Plugins/Media/AvfMedia/Source/AvfMediaCapture/Private/AvfMediaCaptureModule.cpp
@@ -28,14 +28,24 @@ DEFINE_LOG_CATEGORY(LogAvfMediaCapture);
 
 #define LOCTEXT_NAMESPACE "FAvfMediaCaptureFactoryModule"
 
-
-// new API doesn't compile on old IOS sdk
 #if (PLATFORM_IOS && (defined(__IPHONE_17_0) && __IPHONE_OS_VERSION_MAX_ALLOWED >= __IPHONE_17_0))
-	#define USE_NEW_MICROPHONE_API 1
+	#define IOS_17_APIS_AVAILABLE 1
+#else
+	#define IOS_17_APIS_AVAILABLE 0
+#endif
+
+#if (PLATFORM_MAC && (defined(__MAC_14_0) && __MAC_OS_X_VERSION_MAX_ALLOWED >= __MAC_14_0))
+	#define MAC_14_APIS_AVAILABLE 1
 #else
-	#define USE_NEW_MICROPHONE_API 0
+	#define MAC_14_APIS_AVAILABLE 0
 #endif
 
+// New AVCaptureDeviceType APIs don't compile on old IOS or MAC SDKs
+#if (IOS_17_APIS_AVAILABLE || MAC_14_APIS_AVAILABLE)
+	#define USE_NEW_CAPTURE_DEVICE_TYPE_API 1
+#else
+	#define USE_NEW_CAPTURE_DEVICE_TYPE_API 0
+#endif
 
 /**
  * Implements the AvfMediaCapture module.
@@ -50,51 +60,70 @@ public:
 	{
 		SCOPED_AUTORELEASE_POOL
 
-        NSArray* DeviceTypes = nil;
-        
-    #if USE_NEW_MICROPHONE_API
-        DeviceTypes = @[AVCaptureDeviceTypeBuiltInWideAngleCamera, AVCaptureDeviceTypeMicrophone];
-    #else
-        DeviceTypes = @[AVCaptureDeviceTypeBuiltInWideAngleCamera, AVCaptureDeviceTypeBuiltInMicrophone];
-    #endif
+        NSMutableArray* DeviceTypes = [[NSMutableArray alloc] init];
+
+		FString Scheme;
+		AVMediaType MediaType;
+		if (TargetDeviceType == EMediaCaptureDeviceType::Audio)
+		{
+			Scheme = TEXT("audcap://");
+			MediaType = AVMediaTypeAudio;
+			#if USE_NEW_CAPTURE_DEVICE_TYPE_API
+				[DeviceTypes addObject: AVCaptureDeviceTypeMicrophone];
+			#else
+				[DeviceTypes addObject: AVCaptureDeviceTypeBuiltInMicrophone];
+			#endif
+		}
+		else if (TargetDeviceType == EMediaCaptureDeviceType::Video)
+		{
+			Scheme = TEXT("vidcap://");
+			MediaType = AVMediaTypeVideo;
+			[DeviceTypes addObject: AVCaptureDeviceTypeBuiltInWideAngleCamera];
+
+			#if PLATFORM_IOS
+				[DeviceTypes addObject: AVCaptureDeviceTypeBuiltInUltraWideCamera];
+				[DeviceTypes addObject: AVCaptureDeviceTypeBuiltInTelephotoCamera];
+			#endif
+
+			#if USE_NEW_CAPTURE_DEVICE_TYPE_API
+				[DeviceTypes addObject: AVCaptureDeviceTypeExternal];
+			#elif PLATFORM_MAC
+				// AVCaptureDeviceTypeExternalUnknown is only available on macOS 10.15 - 14.0
+				// https://developer.apple.com/documentation/avfoundation/avcapturedevicetypeexternalunknown?language=objc
+				[DeviceTypes addObject: AVCaptureDeviceTypeExternalUnknown]; 
+			#endif
+		}
+
+		if ([DeviceTypes count] == 0)
+		{
+			return;
+		}
 
-		AVCaptureDeviceDiscoverySession* LocalDiscoverySession = [AVCaptureDeviceDiscoverySession discoverySessionWithDeviceTypes:DeviceTypes mediaType:nil position:AVCaptureDevicePositionUnspecified];
-		if(LocalDiscoverySession != nil)
+		AVCaptureDeviceDiscoverySession* LocalDiscoverySession = [AVCaptureDeviceDiscoverySession discoverySessionWithDeviceTypes: DeviceTypes
+			                                                      mediaType: nil
+			                                                      position: AVCaptureDevicePositionUnspecified];
+		if (LocalDiscoverySession != nil)
 		{
 			NSArray<AVCaptureDevice*>* Devices = LocalDiscoverySession.devices;
-			for(uint32 i = 0;i < Devices.count;++i)
+			for(uint32 i = 0; i < Devices.count; ++i)
 			{
 				AVCaptureDevice* AvailableDevice = Devices[i];
-            #if USE_NEW_MICROPHONE_API
-                if(TargetDeviceType == EMediaCaptureDeviceType::Audio && AvailableDevice.deviceType == AVCaptureDeviceTypeMicrophone)
-            #else
-                if(TargetDeviceType == EMediaCaptureDeviceType::Audio && AvailableDevice.deviceType == AVCaptureDeviceTypeBuiltInMicrophone)
-            #endif
+
+				// It's not clear whether external media types will be limited to video.
+				// In which case we double check that the detected device supports the target media type.
+				if (![AvailableDevice hasMediaType: MediaType])
 				{
-					FMediaCaptureDeviceInfo DeviceInfo;
-					
-					DeviceInfo.Type = TargetDeviceType;
-					DeviceInfo.DisplayName = FText::FromString(FString(AvailableDevice.localizedName));
-					DeviceInfo.Url = TEXT("audcap://") + FString(AvailableDevice.uniqueID);
-					DeviceInfo.Info = FString(AvailableDevice.manufacturer);
-					
-					OutDeviceInfos.Add(MoveTemp(DeviceInfo));
+					continue;
 				}
-            #if USE_NEW_MICROPHONE_API
-                else if(TargetDeviceType == EMediaCaptureDeviceType::Video && AvailableDevice.deviceType != AVCaptureDeviceTypeMicrophone)
-            #else
-                else if(TargetDeviceType == EMediaCaptureDeviceType::Video && AvailableDevice.deviceType != AVCaptureDeviceTypeBuiltInMicrophone)
-            #endif
-				{
-					FMediaCaptureDeviceInfo DeviceInfo;
-				
-					DeviceInfo.Type = TargetDeviceType;
-					DeviceInfo.DisplayName = FText::FromString(FString(AvailableDevice.localizedName));
-					DeviceInfo.Url = TEXT("vidcap://") + FString(AvailableDevice.uniqueID);
-					DeviceInfo.Info = FString(AvailableDevice.manufacturer);
+
+				FMediaCaptureDeviceInfo DeviceInfo;
 					
-					OutDeviceInfos.Add(MoveTemp(DeviceInfo));
-				}
+				DeviceInfo.Type = TargetDeviceType;
+				DeviceInfo.DisplayName = FText::FromString(FString(AvailableDevice.localizedName));
+				DeviceInfo.Url = Scheme + FString(AvailableDevice.uniqueID);
+				DeviceInfo.Info = FString(AvailableDevice.manufacturer);
+					
+				OutDeviceInfos.Add(MoveTemp(DeviceInfo));
 			}
 		}
 	}
diff --git a/Engine/Plugins/Media/AvfMedia/Source/AvfMediaCapture/Private/Player/AvfMediaCaptureHelper.cpp b/Engine/Plugins/Media/AvfMedia/Source/AvfMediaCapture/Private/Player/AvfMediaCaptureHelper.cpp
--- a/Engine/Plugins/Media/AvfMedia/Source/AvfMediaCapture/Private/Player/AvfMediaCaptureHelper.cpp
+++ b/Engine/Plugins/Media/AvfMedia/Source/AvfMediaCapture/Private/Player/AvfMediaCaptureHelper.cpp
@@ -14,34 +14,28 @@
 // 3) delete permissions database then reboot mac:
 //  ~/Library/Application\\ Support/com.apple.TCC
 
-// new API doesn't compile on old IOS sdk
-#if (PLATFORM_IOS && (defined(__IPHONE_17_0) && __IPHONE_OS_VERSION_MAX_ALLOWED >= __IPHONE_17_0))
-	#define USE_NEW_MICROPHONE_API 1
-#else
-	#define USE_NEW_MICROPHONE_API 0
-#endif
-
-
 @interface AvfMediaCaptureHelper()<AVCaptureVideoDataOutputSampleBufferDelegate, AVCaptureAudioDataOutputSampleBufferDelegate>
 
-@property (nonatomic,readwrite,assign) AVCaptureSession* 		captureSession;
-@property (nonatomic,readwrite,assign) AVCaptureDevice*			captureDevice;
-@property (nonatomic,readwrite,assign) AVCaptureDeviceInput*	deviceInput;
-@property (nonatomic,readwrite,assign) AVCaptureOutput* 		captureOutput;
-@property (nonatomic,readwrite,assign) dispatch_queue_t 		callBackQueue;
+@property (nonatomic, assign) AVMediaType mediaType;
+@property (nonatomic, assign) AVCaptureSession* captureSession;
+@property (nonatomic, assign) AVCaptureDevice* captureDevice;
+@property (nonatomic, assign) AVCaptureDeviceInput*	deviceInput;
+@property (nonatomic, assign) AVCaptureOutput* captureOutput;
+@property (nonatomic, assign) dispatch_queue_t callBackQueue;
 
-@property (nonatomic,readwrite,assign) void (^sampleBufferCallback)(CMSampleBufferRef);
-@property (nonatomic,readwrite,assign) void (^notificationCallback)(NSNotification* const Notification);
+@property (nonatomic, assign) void (^sampleBufferCallback)(CMSampleBufferRef);
+@property (nonatomic, assign) void (^notificationCallback)(NSNotification* const Notification);
 
 @end
 
 @implementation AvfMediaCaptureHelper
 
-- (instancetype) init
+- (instancetype)init:(AVMediaType)mediaType
 {
 	self = [super init];
-	if(self != nil)
+	if (self != nil)
 	{
+		self.mediaType = mediaType;
 		self.captureSession = nil;
 		self.captureDevice = nil;
 		self.deviceInput = nil;
@@ -53,11 +47,11 @@
 	return self;
 }
 
-- (void) dealloc
+- (void)dealloc
 {
 	[self reset];
 	
-	if(self.callBackQueue != nil)
+	if (self.callBackQueue != nil)
 	{
 		dispatch_release(self.callBackQueue);
 		self.callBackQueue = nil;
@@ -66,60 +60,60 @@
 	[super dealloc];
 }
 
-- (void) reset
+- (void)reset
 {
-	if(self.captureSession.inputs.count > 0 && self.captureSession.outputs.count > 0)
+	if (self.captureSession.inputs.count > 0 && self.captureSession.outputs.count > 0)
 	{
 		[[NSNotificationCenter defaultCenter] removeObserver:self];
 	}
 	
-	if(self.captureSession.isRunning)
+	if (self.captureSession.isRunning)
 	{
 		[self.captureSession stopRunning];
 	}
 	
-	if(self.captureOutput != nil)
+	if (self.captureOutput != nil)
 	{
 		[self.captureSession removeOutput:self.captureOutput];
 		[self.captureOutput release];
 		self.captureOutput = nil;
 	}
 	
-	if(self.deviceInput != nil)
+	if (self.deviceInput != nil)
 	{
 		[self.captureSession removeInput:self.deviceInput];
 		[self.deviceInput release];
 		self.deviceInput = nil;
 	}
 	
-	if(self.captureSession != nil)
+	if (self.captureSession != nil)
 	{
 		[self.captureSession release];
 		self.captureSession = nil;
 	}
 	
-	if(self.captureDevice != nil)
+	if (self.captureDevice != nil)
 	{
 		[self.captureDevice release];
 		self.captureDevice = nil;
 	}
 	
-	if(self.sampleBufferCallback != nil)
+	if (self.sampleBufferCallback != nil)
 	{
 		Block_release(self.sampleBufferCallback);
 		self.sampleBufferCallback = nil;
 	}
 	
-	if(self.notificationCallback != nil)
+	if (self.notificationCallback != nil)
 	{
 		Block_release(self.notificationCallback);
 		self.notificationCallback = nil;
 	}
 }
 
-+ (EAvfMediaCaptureAuthStatus) authorizationStatusForMediaType:(AVMediaType)mediaType
++ (EAvfMediaCaptureAuthStatus)authorizationStatusForMediaType:(AVMediaType)mediaType
 {
-	if(mediaType != AVMediaTypeVideo && mediaType != AVMediaTypeAudio)
+	if (mediaType != AVMediaTypeVideo && mediaType != AVMediaTypeAudio)
 	{
 		return EAvfMediaCaptureAuthStatus::InvalidRequest;
 	}
@@ -128,9 +122,9 @@
 	
 	id entry = nil;
 	
-	if(infoDictionary != nil)
+	if (infoDictionary != nil)
 	{
-		if(mediaType == AVMediaTypeVideo)
+		if (mediaType == AVMediaTypeVideo)
 		{
 			entry = infoDictionary[@"NSCameraUsageDescription"];
 		}
@@ -140,7 +134,7 @@
 		}
 	}
 	
-	if(entry == nil)
+	if (entry == nil)
 	{
 		return EAvfMediaCaptureAuthStatus::MissingInfoPListEntry;
 	}
@@ -148,7 +142,7 @@
 	return (EAvfMediaCaptureAuthStatus)[AVCaptureDevice authorizationStatusForMediaType:mediaType];
 }
 
-+ (EAvfMediaCaptureAuthStatus) requestAcessForMediaType:(AVMediaType)mediaType completionCallback:(void (^)(EAvfMediaCaptureAuthStatus AuthStatus))cbHandler
++ (EAvfMediaCaptureAuthStatus)requestAccessForMediaType:(AVMediaType)mediaType completionCallback:(void (^)(EAvfMediaCaptureAuthStatus AuthStatus))cbHandler
 {
 	// Don't make request if the correct device access key is not in the info.plist otherwise the OS will terminate this app
 	EAvfMediaCaptureAuthStatus authStatus = [AvfMediaCaptureHelper authorizationStatusForMediaType:mediaType];
@@ -175,53 +169,28 @@
 	return authStatus;
 }
 
-- (AVCaptureDevice*) CaptureDeviceWithID:(NSString*)deviceID
-{
-	AVCaptureDevice* foundDevice = nil;
-    NSArray* deviceTypes = nil;
-    
-#if USE_NEW_MICROPHONE_API
-    deviceTypes = @[AVCaptureDeviceTypeBuiltInWideAngleCamera, AVCaptureDeviceTypeMicrophone];
-#else
-    deviceTypes = @[AVCaptureDeviceTypeBuiltInWideAngleCamera, AVCaptureDeviceTypeBuiltInMicrophone];
-#endif
-
-
-	AVCaptureDeviceDiscoverySession* localDiscoverySession = [AVCaptureDeviceDiscoverySession discoverySessionWithDeviceTypes:deviceTypes
-																													mediaType:nil
-																													 position:AVCaptureDevicePositionUnspecified];
-	if(localDiscoverySession != nil)
-	{
-		NSArray<AVCaptureDevice*>* devices = localDiscoverySession.devices;
-		
-		for(AVCaptureDevice* device in devices)
-		{
-			if([device.uniqueID isEqual:deviceID])
-			{
-				foundDevice = device;
-				break;
-			}
-		}
-	}
-
-	return foundDevice;
-}
-
-- (BOOL) setupCaptureSession:(NSString*)deviceID sampleBufferCallback:(void(^)(CMSampleBufferRef sampleBuffer))sampleCallbackBlock notificationCallback:(void(^)(NSNotification* const notification))notificationCallbackBlock
+- (BOOL)setupCaptureSession:(NSString*)deviceID sampleBufferCallback:(void(^)(CMSampleBufferRef sampleBuffer))sampleCallbackBlock notificationCallback:(void(^)(NSNotification* const notification))notificationCallbackBlock
 {
 	[self reset];
 	
-	if(sampleCallbackBlock != nil)
+	if (sampleCallbackBlock != nil)
 	{
 		self.sampleBufferCallback = Block_copy(sampleCallbackBlock);
 	}
-	if(notificationCallbackBlock != nil)
+	if (notificationCallbackBlock != nil)
 	{
 		self.notificationCallback = Block_copy(notificationCallbackBlock);
 	}
 	
-	self.captureDevice = [[self CaptureDeviceWithID:deviceID] retain];
-	if(self.captureDevice != nil)
+	self.captureDevice = [[AVCaptureDevice deviceWithUniqueID: deviceID] retain];
+
+	// Check that the chosen device supports the specified media type
+	if (![self.captureDevice hasMediaType: self.mediaType])
+	{
+		return NO;
+	}
+	
+	if (self.captureDevice != nil)
 	{
 		self.captureSession = [[AVCaptureSession alloc] init];
 		
@@ -232,21 +201,17 @@
 		{
 			self.deviceInput = [[AVCaptureDeviceInput alloc] initWithDevice:self.captureDevice error:nil];
 			
-			if(self.deviceInput != nil && [self.captureSession canAddInput:self.deviceInput])
+			if (self.deviceInput != nil && [self.captureSession canAddInput:self.deviceInput])
 			{
 				[self.captureSession addInput:self.deviceInput];
 			}
 		}
 		
 		// If we have an input add output capture
-		if(self.captureSession.inputs.count > 0)
+		if (self.captureSession.inputs.count > 0)
 		{
             // Video or Audio Device - need the correct output
-#if USE_NEW_MICROPHONE_API
-            if(self.captureDevice.deviceType != AVCaptureDeviceTypeMicrophone)
-#else
-            if(self.captureDevice.deviceType != AVCaptureDeviceTypeBuiltInMicrophone)
-#endif
+            if (self.mediaType == AVMediaTypeVideo)
 			{
 				AVCaptureVideoDataOutput* videoOutput = [[AVCaptureVideoDataOutput alloc] init];
 				
@@ -288,14 +253,14 @@
 				self.captureOutput = audioOutput;
 			}
 			
-			if(self.captureOutput != nil && [self.captureSession canAddOutput:self.captureOutput])
+			if (self.captureOutput != nil && [self.captureSession canAddOutput:self.captureOutput])
 			{
 				[self.captureSession addOutput:self.captureOutput];
 			}
 		}
 	}
 	
-	if(self.captureSession.inputs.count > 0 && self.captureSession.outputs.count > 0)
+	if (self.captureSession.inputs.count > 0 && self.captureSession.outputs.count > 0)
 	{
 		NSNotificationCenter* sharedCentre = [NSNotificationCenter defaultCenter];
 
@@ -306,92 +271,83 @@
 		[sharedCentre addObserver:self selector:@selector(captureNotification:) name:AVCaptureSessionRuntimeErrorNotification object:self.captureSession];
 		
 #if PLATFORM_MAC && WITH_EDITOR
-		if([self.captureOutput isKindOfClass:[AVCaptureVideoDataOutput class]])
+		if ([self.captureOutput isKindOfClass:[AVCaptureVideoDataOutput class]])
 		{
 			[sharedCentre addObserver:self selector:@selector(captureNotification:) name:NSApplicationDidBecomeActiveNotification object:nil];
 			[sharedCentre addObserver:self selector:@selector(captureNotification:) name:NSApplicationWillResignActiveNotification object:nil];
 		}
 #endif
 
-		return true;
+		return YES;
 	}
 
-	return false;
+	return NO;
 }
 
-- (void) stopCaptureSession
+- (void)stopCaptureSession
 {
-	if(self.captureSession.isRunning)
+	if (self.captureSession.isRunning)
 	{
 		[self.captureSession stopRunning];
 	}
 }
 
-- (void) startCaptureSession
+- (void)startCaptureSession
 {
-	if(!self.captureSession.isRunning)
+	if (!self.captureSession.isRunning)
 	{
 		[self.captureSession startRunning];
 	}
 }
 
-- (BOOL) isCaptureRunning
+- (BOOL)isCaptureRunning
 {
 	return self.captureSession.isRunning;
 }
 
 - (void)captureOutput:(AVCaptureOutput *)output didOutputSampleBuffer:(CMSampleBufferRef)sampleBuffer fromConnection:(AVCaptureConnection*)connection
 {
-	if(self.sampleBufferCallback != nil && sampleBuffer != NULL)
+	if (self.sampleBufferCallback != nil && sampleBuffer != NULL)
 	{
 		self.sampleBufferCallback(sampleBuffer);
 	}
 }
 
-- (void) captureNotification:(NSNotification*)notification
+- (void)captureNotification:(NSNotification*)notification
 {
-	if(self.notificationCallback != nil && notification != nil)
+	if (self.notificationCallback != nil && notification != nil)
 	{
 		self.notificationCallback(notification);
 	}
 }
 
-- (AVMediaType) getCaptureDeviceMediaType
+- (AVMediaType)getCaptureDeviceMediaType
 {
-	AVMediaType type = nil;
-	if(self.captureDevice != nil)
-	{
-#if USE_NEW_MICROPHONE_API
-		if(self.captureDevice.deviceType != AVCaptureDeviceTypeMicrophone)
-#else
-        if(self.captureDevice.deviceType != AVCaptureDeviceTypeBuiltInMicrophone)
-#endif
-		{
-			type = AVMediaTypeVideo;
-		}
-		else
-		{
-			type = AVMediaTypeAudio;
-		}
-	}
-	return type;
+	return self.mediaType;
 }
 
-- (NSString*) getCaptureDeviceName
+- (NSString*)getCaptureDeviceName
 {
 	return self.captureDevice.localizedName;
 }
 
-- (NSArray<AVCaptureDeviceFormat*>*) getCaptureDeviceAvailableFormats
+- (NSArray<AVCaptureDeviceFormat*>*)getCaptureDeviceAvailableFormats
 {
-	return self.captureDevice.formats;
+	AVMediaType helperMediaType = self.mediaType;
+	NSArray<AVCaptureDeviceFormat*>* filtered = [self.captureDevice.formats filteredArrayUsingPredicate: [NSPredicate predicateWithBlock:^BOOL(AVCaptureDeviceFormat* format, NSDictionary* bindings)
+	{
+		return [format.mediaType isEqualToString: helperMediaType];
+	}]];
+
+	return filtered;
 }
 
-- (NSInteger) getCaptureDeviceActiveFormatIndex
+- (NSInteger)getCaptureDeviceActiveFormatIndex
 {
-	for(NSInteger formatIndex = 0;formatIndex < self.captureDevice.formats.count;++formatIndex)
+	NSArray<AVCaptureDeviceFormat*>* formats = [self getCaptureDeviceAvailableFormats];
+	for (NSInteger formatIndex = 0; formatIndex < formats.count; ++formatIndex)
 	{
-		if([self.captureDevice.formats[formatIndex] isEqual:self.captureDevice.activeFormat])
+		if ([formats[formatIndex] isEqual:self.captureDevice.activeFormat])
 		{
 			return formatIndex;
 		}
@@ -399,11 +355,12 @@
 	return INDEX_NONE;
 }
 
-- (BOOL) setCaptureDeviceActiveFormatIndex:(NSInteger)formatIdx
+- (BOOL)setCaptureDeviceActiveFormatIndex:(NSInteger)formatIdx
 {
-	if(formatIdx >= 0 && formatIdx < self.captureDevice.formats.count && [self.captureDevice lockForConfiguration:nil])
+	NSArray<AVCaptureDeviceFormat*>* formats = [self getCaptureDeviceAvailableFormats];
+	if (formatIdx >= 0 && formatIdx < formats.count && [self.captureDevice lockForConfiguration: nil])
 	{
-		self.captureDevice.activeFormat = self.captureDevice.formats[formatIdx];
+		self.captureDevice.activeFormat = formats[formatIdx];
 		[self.captureDevice unlockForConfiguration];
 		return YES;
 	}
diff --git a/Engine/Plugins/Media/AvfMedia/Source/AvfMediaCapture/Private/Player/AvfMediaCaptureHelper.h b/Engine/Plugins/Media/AvfMedia/Source/AvfMediaCapture/Private/Player/AvfMediaCaptureHelper.h
--- a/Engine/Plugins/Media/AvfMedia/Source/AvfMediaCapture/Private/Player/AvfMediaCaptureHelper.h
+++ b/Engine/Plugins/Media/AvfMedia/Source/AvfMediaCapture/Private/Player/AvfMediaCaptureHelper.h
@@ -27,23 +27,26 @@ enum class EAvfMediaCaptureAuthStatus : uint32
 @interface AvfMediaCaptureHelper : NSObject
 
 // Returns current status or error
-+ (EAvfMediaCaptureAuthStatus) authorizationStatusForMediaType:(AVMediaType)mediaType;
++ (EAvfMediaCaptureAuthStatus)authorizationStatusForMediaType:(AVMediaType)mediaType;
 
 // Returns current status or error and requests access if not determined
-+ (EAvfMediaCaptureAuthStatus) requestAcessForMediaType:(AVMediaType)mediaType completionCallback:(void(^)(EAvfMediaCaptureAuthStatus AuthStatus))cbHandler;
++ (EAvfMediaCaptureAuthStatus)requestAccessForMediaType:(AVMediaType)mediaType completionCallback:(void(^)(EAvfMediaCaptureAuthStatus AuthStatus))cbHandler;
 
-- (BOOL) setupCaptureSession:(NSString*)deviceID sampleBufferCallback:(void(^)(CMSampleBufferRef sampleBuffer))sampleCallbackBlock
-												 notificationCallback:(void(^)(NSNotification* const notification))notificationCallbackBlock;
-- (void) stopCaptureSession;
-- (void) startCaptureSession;
-- (BOOL) isCaptureRunning;
+- (instancetype)init:(AVMediaType)mediaType;
 
-- (NSString*) 	getCaptureDeviceName;
-- (AVMediaType) getCaptureDeviceMediaType;
+- (BOOL)setupCaptureSession:(NSString*)deviceID
+        sampleBufferCallback:(void(^)(CMSampleBufferRef sampleBuffer))sampleCallbackBlock
+        notificationCallback:(void(^)(NSNotification* const notification))notificationCallbackBlock;
+- (void)stopCaptureSession;
+- (void)startCaptureSession;
+- (BOOL)isCaptureRunning;
 
-- (NSArray<AVCaptureDeviceFormat*>*) getCaptureDeviceAvailableFormats;
+- (NSString*)getCaptureDeviceName;
+- (AVMediaType)getCaptureDeviceMediaType;
 
-- (NSInteger) 	getCaptureDeviceActiveFormatIndex;
-- (BOOL) 		setCaptureDeviceActiveFormatIndex:(NSInteger)formatIdx;
+- (NSArray<AVCaptureDeviceFormat*>*)getCaptureDeviceAvailableFormats;
+
+- (NSInteger)getCaptureDeviceActiveFormatIndex;
+- (BOOL)setCaptureDeviceActiveFormatIndex:(NSInteger)formatIdx;
 
 @end
diff --git a/Engine/Plugins/Media/AvfMedia/Source/AvfMediaCapture/Private/Player/AvfMediaCapturePlayer.cpp b/Engine/Plugins/Media/AvfMedia/Source/AvfMediaCapture/Private/Player/AvfMediaCapturePlayer.cpp
--- a/Engine/Plugins/Media/AvfMedia/Source/AvfMediaCapture/Private/Player/AvfMediaCapturePlayer.cpp
+++ b/Engine/Plugins/Media/AvfMedia/Source/AvfMediaCapture/Private/Player/AvfMediaCapturePlayer.cpp
@@ -184,7 +184,7 @@ void FAvfMediaCapturePlayer::HandleAuthStatusError(EAvfMediaCaptureAuthStatus Au
 	}
 }
 
-void FAvfMediaCapturePlayer::CreateCaptureSession(NSString* deviceIDString)
+void FAvfMediaCapturePlayer::CreateCaptureSession(NSString* DeviceIDString, AVMediaType MediaType)
 {
 	SCOPED_AUTORELEASE_POOL;
 	
@@ -193,8 +193,8 @@ void FAvfMediaCapturePlayer::CreateCaptureSession(NSString* deviceIDString)
 	if(MediaCaptureHelper == nil)
 	{
 		// Assuming only one Capture at the moment
-		MediaCaptureHelper = [[AvfMediaCaptureHelper alloc] init];
-		BOOL bResult = [MediaCaptureHelper setupCaptureSession:deviceIDString sampleBufferCallback:^(CMSampleBufferRef SampleBuffer)
+		MediaCaptureHelper = [[AvfMediaCaptureHelper alloc] init: MediaType];
+		BOOL bResult = [MediaCaptureHelper setupCaptureSession: DeviceIDString sampleBufferCallback: ^(CMSampleBufferRef SampleBuffer)
 		{
 			this->NewSampleBufferAvailable(SampleBuffer);
 		}
@@ -240,20 +240,20 @@ bool FAvfMediaCapturePlayer::Open(const FString& Url, const IMediaOptions* Optio
 		{
 			case EAvfMediaCaptureAuthStatus::Authorized:
 			{
-				CreateCaptureSession(nsDeviceID);
+				CreateCaptureSession(nsDeviceID, MediaType);
 				break;
 			}
 			case EAvfMediaCaptureAuthStatus::NotDetermined:
 			{
 				TWeakPtr<FAvfMediaCapturePlayer, ESPMode::ThreadSafe> WeakPtr = this->AsWeak();
-				[AvfMediaCaptureHelper requestAcessForMediaType:MediaType completionCallback:^(EAvfMediaCaptureAuthStatus ResultStatus)
+				[AvfMediaCaptureHelper requestAccessForMediaType:MediaType completionCallback:^(EAvfMediaCaptureAuthStatus ResultStatus)
 				{
 					TSharedPtr<FAvfMediaCapturePlayer, ESPMode::ThreadSafe> StrongPtr = WeakPtr.Pin();
 					if(StrongPtr.IsValid())
 					{
 						if(ResultStatus == EAvfMediaCaptureAuthStatus::Authorized)
 						{
-							StrongPtr->CreateCaptureSession(nsDeviceID);
+							StrongPtr->CreateCaptureSession(nsDeviceID, MediaType);
 						}
 						else
 						{
@@ -354,12 +354,25 @@ int32 FAvfMediaCapturePlayer::GetNumTrackFormats(EMediaTrackType TrackType, int3
 		FScopeLock Lock(&CriticalSection);
 		if(MediaCaptureHelper != nil)
 		{
-			AVMediaType MediaType = [MediaCaptureHelper getCaptureDeviceMediaType];
-			if	(	(TrackType == EMediaTrackType::Video && MediaType == AVMediaTypeVideo) ||
-					(TrackType == EMediaTrackType::Audio && MediaType == AVMediaTypeAudio)
-				)
+			AVMediaType MediaType;
+			switch (TrackType) {
+			case EMediaTrackType::Audio:
+				MediaType = AVMediaTypeAudio;
+				break;
+			case EMediaTrackType::Video:
+				MediaType = AVMediaTypeVideo;
+				break;
+			default:
+				return 0;
+			}
+
+			if ([[MediaCaptureHelper getCaptureDeviceMediaType] isEqualToString: MediaType])
 			{
-				return [MediaCaptureHelper getCaptureDeviceAvailableFormats].count;
+				return [[MediaCaptureHelper getCaptureDeviceAvailableFormats] count];
+			}
+			else
+			{
+				return 0;
 			}
 		}
 	}
@@ -372,7 +385,7 @@ int32 FAvfMediaCapturePlayer::GetSelectedTrack(EMediaTrackType TrackType) const
 	SCOPED_AUTORELEASE_POOL;
 	
 	FScopeLock Lock(&CriticalSection);
-	if(MediaCaptureHelper != nil)
+	if (MediaCaptureHelper != nil)
 	{
 		AVMediaType MediaType = [MediaCaptureHelper getCaptureDeviceMediaType];
 		if	(	(TrackType == EMediaTrackType::Video && MediaType == AVMediaTypeVideo) ||
@@ -394,10 +407,10 @@ int32 FAvfMediaCapturePlayer::GetTrackFormat(EMediaTrackType TrackType, int32 Tr
 {
 	SCOPED_AUTORELEASE_POOL;
 	
-	if(TrackIndex == 0)
+	if (TrackIndex == 0)
 	{
 		FScopeLock Lock(&CriticalSection);
-		if(MediaCaptureHelper != nil)
+		if (MediaCaptureHelper != nil)
 		{
 			AVMediaType MediaType = [MediaCaptureHelper getCaptureDeviceMediaType];
 			if	(	(TrackType == EMediaTrackType::Video && MediaType == AVMediaTypeVideo) ||
@@ -427,13 +440,13 @@ bool FAvfMediaCapturePlayer::GetVideoTrackFormat(int32 TrackIndex, int32 FormatI
 	
 	FScopeLock Lock(&CriticalSection);
 	
-	if(TrackIndex == 0 && FormatIndex >= 0)
+	if (TrackIndex == 0 && FormatIndex >= 0)
 	{
 		NSArray<AVCaptureDeviceFormat*>* deviceFormats = [MediaCaptureHelper getCaptureDeviceAvailableFormats];
-		if(FormatIndex < deviceFormats.count)
+		if (FormatIndex < deviceFormats.count)
 		{
 			AVCaptureDeviceFormat* availableFormat = deviceFormats[FormatIndex];
-			if(CMFormatDescriptionGetMediaType(availableFormat.formatDescription) == kCMMediaType_Video)
+			if (CMFormatDescriptionGetMediaType(availableFormat.formatDescription) == kCMMediaType_Video)
 			{
 				float fMinRate = FLT_MAX;
 				float fMaxRate = FLT_MIN;
@@ -442,18 +455,18 @@ bool FAvfMediaCapturePlayer::GetVideoTrackFormat(int32 TrackIndex, int32 FormatI
 				for(uint32 i = 0;i < rateRanges.count;++i)
 				{
 					AVFrameRateRange* range = rateRanges[i];
-					if(range.maxFrameRate > fMaxRate)
+					if (range.maxFrameRate > fMaxRate)
 					{
 						fMaxRate = range.maxFrameRate;
 					}
-					if(range.minFrameRate < fMinRate)
+					if (range.minFrameRate < fMinRate)
 					{
 						fMinRate = range.minFrameRate;
 					}
 				}
 				
 				// Missing rate info - set a sensible value
-				if(fMinRate > fMaxRate)
+				if (fMinRate > fMaxRate)
 				{
 					fMinRate = fMaxRate = 30.f;
 				}
@@ -483,12 +496,12 @@ bool FAvfMediaCapturePlayer::SetTrackFormat(EMediaTrackType TrackType, int32 Tra
 {
 	SCOPED_AUTORELEASE_POOL;
 	
-	if(TrackIndex == 0)
+	if (TrackIndex == 0)
 	{
 		FScopeLock Lock(&CriticalSection);
-		if(TrackType == EMediaTrackType::Video || TrackType == EMediaTrackType::Audio)
+		if (TrackType == EMediaTrackType::Video || TrackType == EMediaTrackType::Audio)
 		{
-			if([MediaCaptureHelper setCaptureDeviceActiveFormatIndex:FormatIndex])
+			if ([MediaCaptureHelper setCaptureDeviceActiveFormatIndex:FormatIndex])
 			{
 				return true;
 			}
@@ -520,9 +533,9 @@ EMediaState FAvfMediaCapturePlayer::GetState() const
 	SCOPED_AUTORELEASE_POOL;
 	
 	FScopeLock Lock(&CriticalSection);
-	if(MediaCaptureHelper != nil)
+	if (MediaCaptureHelper != nil)
 	{
-		if(CurrentRate == 0.f && ![MediaCaptureHelper isCaptureRunning])
+		if (CurrentRate == 0.f && ![MediaCaptureHelper isCaptureRunning])
 		{
 			return EMediaState::Paused;
 		}
@@ -575,13 +588,13 @@ bool FAvfMediaCapturePlayer::SetRate(float Rate)
 	SCOPED_AUTORELEASE_POOL
 	
 	FScopeLock Lock(&CriticalSection);
-	if(Rate != CurrentRate)
+	if (Rate != CurrentRate)
 	{
-		if(Rate == 0.f && CurrentRate != 0.f)
+		if (Rate == 0.f && CurrentRate != 0.f)
 		{
 			[MediaCaptureHelper stopCaptureSession];
 		}
-		else if(Rate != 0.f && CurrentRate == 0.f)
+		else if (Rate != 0.f && CurrentRate == 0.f)
 		{
 			[MediaCaptureHelper startCaptureSession];
 		}
@@ -594,32 +607,32 @@ bool FAvfMediaCapturePlayer::SetRate(float Rate)
 void FAvfMediaCapturePlayer::CaptureSystemNotification(NSNotification* Notification)
 {
 	NSString* NotificationName = Notification.name;
-	if(NotificationName == AVCaptureSessionRuntimeErrorNotification)
+	if (NotificationName == AVCaptureSessionRuntimeErrorNotification)
 	{
 		Close();
 	}
-	else if(NotificationName == AVCaptureSessionDidStartRunningNotification)
+	else if (NotificationName == AVCaptureSessionDidStartRunningNotification)
 	{
 		EventSink.ReceiveMediaEvent(EMediaEvent::PlaybackResumed);
 	}
-	else if(NotificationName == AVCaptureSessionDidStopRunningNotification)
+	else if (NotificationName == AVCaptureSessionDidStopRunningNotification)
 	{
 		EventSink.ReceiveMediaEvent(EMediaEvent::PlaybackSuspended);
 	}
-	else if(NotificationName == AVCaptureSessionWasInterruptedNotification)
+	else if (NotificationName == AVCaptureSessionWasInterruptedNotification)
 	{
 		EventSink.ReceiveMediaEvent(EMediaEvent::PlaybackSuspended);
 	}
-	else if(NotificationName == AVCaptureSessionInterruptionEndedNotification)
+	else if (NotificationName == AVCaptureSessionInterruptionEndedNotification)
 	{
 		EventSink.ReceiveMediaEvent(EMediaEvent::PlaybackResumed);
 	}
 #if PLATFORM_MAC && WITH_EDITOR
-	else if(NotificationName == NSApplicationDidBecomeActiveNotification)
+	else if (NotificationName == NSApplicationDidBecomeActiveNotification)
 	{
 		ThrottleDuration = 0.0;
 	}
-	else if(NotificationName == NSApplicationWillResignActiveNotification)
+	else if (NotificationName == NSApplicationWillResignActiveNotification)
 	{
 		ThrottleDuration = 1.0;
 	}
@@ -628,28 +641,28 @@ void FAvfMediaCapturePlayer::CaptureSystemNotification(NSNotification* Notificat
 
 void FAvfMediaCapturePlayer::NewSampleBufferAvailable(CMSampleBufferRef SampleBuffer)
 {
-	if(!CMSampleBufferIsValid(SampleBuffer) || !CMSampleBufferDataIsReady(SampleBuffer))
+	if (!CMSampleBufferIsValid(SampleBuffer) || !CMSampleBufferDataIsReady(SampleBuffer))
 	{
 		return;
 	}
 	
 	FScopeLock Lock(&CriticalSection);
 	
-	if(CurrentRate == 0.f)
+	if (CurrentRate == 0.f)
 	{
 		return;
 	}
 
 	// Allow each media type to compute the duration as CMSampleBuffer doesn't always report valid timings
 	CMFormatDescriptionRef FormatDescription = CMSampleBufferGetFormatDescription(SampleBuffer);
-	if(FormatDescription != NULL)
+	if (FormatDescription != NULL)
 	{
 		CMMediaType MediaType = CMFormatDescriptionGetMediaType(FormatDescription);
-		if(MediaType == kCMMediaType_Audio)
+		if (MediaType == kCMMediaType_Audio)
 		{
 			ProcessSampleBufferAudio(SampleBuffer);
 		}
-		else if(MediaType == kCMMediaType_Video)
+		else if (MediaType == kCMMediaType_Video)
 		{
 			ProcessSampleBufferVideo(SampleBuffer);
 		}
@@ -662,22 +675,22 @@ FTimespan FAvfMediaCapturePlayer::UpdateInternalTime(CMSampleBufferRef SampleBuf
 		
 	{
 		CMTime SampleBufferDuration = CMSampleBufferGetDuration(SampleBuffer);
-		if((SampleBufferDuration.flags & kCMTimeFlags_Valid) == kCMTimeFlags_Valid)
+		if ((SampleBufferDuration.flags & kCMTimeFlags_Valid) == kCMTimeFlags_Valid)
 		{
 			SampleDuration = FTimespan::FromSeconds(CMTimeGetSeconds(SampleBufferDuration));
 		}
 	}
 	
-	if(SampleDuration.IsZero())
+	if (SampleDuration.IsZero())
 	{
 		SampleDuration = ComputedBufferDuration;
 	}
 	
-	if(CurrentTime == FTimespan::MinValue())
+	if (CurrentTime == FTimespan::MinValue())
 	{
 		CurrentTime = FTimespan::Zero();
 	}
-	else if(CurrentTime >= FTimespan::Zero())
+	else if (CurrentTime >= FTimespan::Zero())
 	{
 		CurrentTime += SampleDuration;
 	}
@@ -694,7 +707,7 @@ void FAvfMediaCapturePlayer::ProcessSampleBufferVideo(CMSampleBufferRef SampleBu
 #if PLATFORM_MAC && WITH_EDITOR
 	// Throttle frames pasted to the editor when not in focus
 	double PlatformTime = FPlatformTime::Seconds();
-	if(ThrottleDuration > 0.0 && (PlatformTime - LastConsumedTimeStamp) < ThrottleDuration)
+	if (ThrottleDuration > 0.0 && (PlatformTime - LastConsumedTimeStamp) < ThrottleDuration)
 	{
 		return;
 	}
@@ -702,7 +715,7 @@ void FAvfMediaCapturePlayer::ProcessSampleBufferVideo(CMSampleBufferRef SampleBu
 #endif
 
 	CVPixelBufferRef PixelBuffer = CMSampleBufferGetImageBuffer(SampleBuffer);
-	if(PixelBuffer)
+	if (PixelBuffer)
 	{
 		check(CMSampleBufferGetNumSamples(SampleBuffer) == 1);
 		check(!CVPixelBufferIsPlanar(PixelBuffer));
@@ -754,14 +767,14 @@ void FAvfMediaCapturePlayer::ProcessSampleBufferAudio(CMSampleBufferRef SampleBu
 	
 	// Avoid using CMSampleBufferGetAudioBufferListWithRetainedBlockBuffer as that appears to copy data - don't want to incurr the extra copy before we even copy out into the UE structure
 	CMFormatDescriptionRef FormatDescription = CMSampleBufferGetFormatDescription(SampleBuffer);
-	if(FormatDescription != NULL && CMFormatDescriptionGetMediaType(FormatDescription) == kCMMediaType_Audio)
+	if (FormatDescription != NULL && CMFormatDescriptionGetMediaType(FormatDescription) == kCMMediaType_Audio)
 	{
 		CMAudioFormatDescriptionRef AudioFormatDescription = (CMAudioFormatDescriptionRef)FormatDescription;
 		AudioStreamBasicDescription const* ASBD = CMAudioFormatDescriptionGetStreamBasicDescription(AudioFormatDescription);
 		check(ASBD);
 		
 		CMBlockBufferRef BlockBuffer = CMSampleBufferGetDataBuffer(SampleBuffer);
-		if(BlockBuffer != NULL)
+		if (BlockBuffer != NULL)
 		{
 			const size_t BufferSize = CMBlockBufferGetDataLength(BlockBuffer);
 			const size_t BitDepth = ASBD->mBitsPerChannel;
@@ -775,22 +788,22 @@ void FAvfMediaCapturePlayer::ProcessSampleBufferAudio(CMSampleBufferRef SampleBu
 			
 			// Get valid engine format
 			EMediaAudioSampleFormat SampleFormat = EMediaAudioSampleFormat::Undefined;
-			if((ASBD->mFormatFlags & kLinearPCMFormatFlagIsFloat) != 0)
+			if ((ASBD->mFormatFlags & kLinearPCMFormatFlagIsFloat) != 0)
 			{
-				if(BitDepth == 32)		SampleFormat = EMediaAudioSampleFormat::Float;
-				else if(BitDepth == 64)	SampleFormat = EMediaAudioSampleFormat::Double;
+				if (BitDepth == 32)		SampleFormat = EMediaAudioSampleFormat::Float;
+				else if (BitDepth == 64)	SampleFormat = EMediaAudioSampleFormat::Double;
 			}
 			else
 			{
-				if(BitDepth == 32)		SampleFormat = EMediaAudioSampleFormat::Int32;
-				else if(BitDepth == 24) SampleFormat = EMediaAudioSampleFormat::Undefined; // 24-bit int is not supported by UE
-				else if(BitDepth == 16)	SampleFormat = EMediaAudioSampleFormat::Int16;
-				else if(BitDepth == 8)	SampleFormat = EMediaAudioSampleFormat::Int8;
+				if (BitDepth == 32)		SampleFormat = EMediaAudioSampleFormat::Int32;
+				else if (BitDepth == 24) SampleFormat = EMediaAudioSampleFormat::Undefined; // 24-bit int is not supported by UE
+				else if (BitDepth == 16)	SampleFormat = EMediaAudioSampleFormat::Int16;
+				else if (BitDepth == 8)	SampleFormat = EMediaAudioSampleFormat::Int8;
 			}
 
 			FTimespan SampleDuration = UpdateInternalTime(SampleBuffer, FTimespan::FromSeconds(((double)NumFrames) / ASBD->mSampleRate));
 			
-			if(SampleFormat != EMediaAudioSampleFormat::Undefined)
+			if (SampleFormat != EMediaAudioSampleFormat::Undefined)
 			{
 				TSharedRef<FAvfMediaAudioSample, ESPMode::ThreadSafe> AudioSample = AudioSamplePool.AcquireShared();
 				if (AudioSample->Initialize(BufferSize,
@@ -811,12 +824,12 @@ void FAvfMediaCapturePlayer::ProcessSampleBufferAudio(CMSampleBufferRef SampleBu
 						size_t LengthAtOffset = 0;
 						size_t BlockTotalSize = 0;
 
-						if(CMBlockBufferGetDataPointer(BlockBuffer, CurrentOffset, &LengthAtOffset, &BlockTotalSize, &SrcData) != noErr)
+						if (CMBlockBufferGetDataPointer(BlockBuffer, CurrentOffset, &LengthAtOffset, &BlockTotalSize, &SrcData) != noErr)
 						{
 							break; // can't get data pointer force exit while loop
 						}
 
-						if(LengthAtOffset == 0)
+						if (LengthAtOffset == 0)
 						{
 							break; // no data force exit while loop
 						}
@@ -825,9 +838,9 @@ void FAvfMediaCapturePlayer::ProcessSampleBufferAudio(CMSampleBufferRef SampleBu
 						checkf(BlockTotalSize == BufferSize, TEXT("Total CMBlockBuffer size does not match data pointer total sizes"));
 						checkf(CMBlockBufferIsRangeContiguous(BlockBuffer, CurrentOffset, LengthAtOffset), TEXT("Returned data pointer for CMBlockBuffer range is not contiguous"));
 
-						if(SrcData != NULL)
+						if (SrcData != NULL)
 						{
-							if((ASBD->mFormatFlags & kAudioFormatFlagIsNonInterleaved) == 0 || NumChannels == 1)
+							if ((ASBD->mFormatFlags & kAudioFormatFlagIsNonInterleaved) == 0 || NumChannels == 1)
 							{
 								// Interleaved data OR single channel - direct copy
 								FMemory::Memcpy(DestData + CurrentOffset, SrcData, LengthAtOffset);
@@ -838,11 +851,11 @@ void FAvfMediaCapturePlayer::ProcessSampleBufferAudio(CMSampleBufferRef SampleBu
 								size_t SrcNumFrames = LengthAtOffset / (BitDepth / 8);
 								size_t SrcFrameOffset = CurrentOffset / (BitDepth / 8);
 								
-								if(SampleFormat == EMediaAudioSampleFormat::Float) 			InterleaveCopy<float>(DestData, SrcData, NumChannels, NumFrames, SrcFrameOffset, SrcNumFrames);
-								else if(SampleFormat == EMediaAudioSampleFormat::Double) 	InterleaveCopy<double>(DestData, SrcData, NumChannels, NumFrames, SrcFrameOffset, SrcNumFrames);
-								else if(SampleFormat == EMediaAudioSampleFormat::Int32) 	InterleaveCopy<int32>(DestData, SrcData, NumChannels, NumFrames, SrcFrameOffset, SrcNumFrames);
-								else if(SampleFormat == EMediaAudioSampleFormat::Int16) 	InterleaveCopy<int16>(DestData, SrcData, NumChannels, NumFrames, SrcFrameOffset, SrcNumFrames);
-								else if(SampleFormat == EMediaAudioSampleFormat::Int8) 		InterleaveCopy<int8>(DestData, SrcData, NumChannels, NumFrames, SrcFrameOffset, SrcNumFrames);
+								if (SampleFormat == EMediaAudioSampleFormat::Float) 			InterleaveCopy<float>(DestData, SrcData, NumChannels, NumFrames, SrcFrameOffset, SrcNumFrames);
+								else if (SampleFormat == EMediaAudioSampleFormat::Double) 	InterleaveCopy<double>(DestData, SrcData, NumChannels, NumFrames, SrcFrameOffset, SrcNumFrames);
+								else if (SampleFormat == EMediaAudioSampleFormat::Int32) 	InterleaveCopy<int32>(DestData, SrcData, NumChannels, NumFrames, SrcFrameOffset, SrcNumFrames);
+								else if (SampleFormat == EMediaAudioSampleFormat::Int16) 	InterleaveCopy<int16>(DestData, SrcData, NumChannels, NumFrames, SrcFrameOffset, SrcNumFrames);
+								else if (SampleFormat == EMediaAudioSampleFormat::Int8) 		InterleaveCopy<int8>(DestData, SrcData, NumChannels, NumFrames, SrcFrameOffset, SrcNumFrames);
 								else break;	// unsupported force exit while loop
 							}
 						}
@@ -850,7 +863,7 @@ void FAvfMediaCapturePlayer::ProcessSampleBufferAudio(CMSampleBufferRef SampleBu
 						CurrentOffset += LengthAtOffset;
 					}
 					
-					if(CurrentOffset < BufferSize)
+					if (CurrentOffset < BufferSize)
 					{
 						// Error case - clear the output buffer to avoid unwanted noise
 						FMemory::Memset(DestData, 0, BufferSize);
diff --git a/Engine/Plugins/Media/AvfMedia/Source/AvfMediaCapture/Private/Player/AvfMediaCapturePlayer.h b/Engine/Plugins/Media/AvfMedia/Source/AvfMediaCapture/Private/Player/AvfMediaCapturePlayer.h
--- a/Engine/Plugins/Media/AvfMedia/Source/AvfMediaCapture/Private/Player/AvfMediaCapturePlayer.h
+++ b/Engine/Plugins/Media/AvfMedia/Source/AvfMediaCapture/Private/Player/AvfMediaCapturePlayer.h
@@ -80,7 +80,7 @@ protected:
 	
 private:
 
-	void CreateCaptureSession(NSString* deviceIDString);
+	void CreateCaptureSession(NSString* deviceIDString, AVMediaType MediaType);
 	void CaptureSystemNotification(NSNotification* Notification);
 	void HandleAuthStatusError(EAvfMediaCaptureAuthStatus AuthStatus, AVMediaType MediaType);
 
@@ -102,20 +102,20 @@ private:
 	AvfMediaCaptureHelper* MediaCaptureHelper;
 	
 	/** Reuse Sample Pool Objects */
-	FAvfMediaAudioSamplePool	AudioSamplePool;
-	FAvfMediaTextureSamplePool	VideoSamplePool;
+	FAvfMediaAudioSamplePool AudioSamplePool;
+	FAvfMediaTextureSamplePool VideoSamplePool;
 	
 	/** Internal CoreVideo Metal Texture cache to handle AVFoundation resource pooling and optimisations */
-	CVMetalTextureCacheRef		MetalTextureCache;
+	CVMetalTextureCacheRef MetalTextureCache;
 	
 	/** Media Playback Info and Control data */
-	float 		CurrentRate;
-	FTimespan 	CurrentTime;
-	FString 	URL;
+	float CurrentRate;
+	FTimespan CurrentTime;
+	FString URL;
 
 #if PLATFORM_MAC && WITH_EDITOR
-	double		ThrottleDuration;
-	double		LastConsumedTimeStamp;
+	double ThrottleDuration;
+	double LastConsumedTimeStamp;
 #endif
 
 	mutable FCriticalSection CriticalSection;
