AvfMedia: workaround for randomly failing prerolling

#rnx

diff --git a/Engine/Plugins/Media/AvfMedia/Source/AvfMedia/Private/Player/AvfMediaPlayer.cpp b/Engine/Plugins/Media/AvfMedia/Source/AvfMedia/Private/Player/AvfMediaPlayer.cpp
--- a/Engine/Plugins/Media/AvfMedia/Source/AvfMedia/Private/Player/AvfMediaPlayer.cpp
+++ b/Engine/Plugins/Media/AvfMedia/Source/AvfMedia/Private/Player/AvfMediaPlayer.cpp
@@ -65,7 +65,7 @@
 	if (Self)
 	{
 		MediaPlayer = InPlayer;
-	}	
+	}
 	return Self;
 }
 
@@ -128,14 +128,14 @@
 {
 	{
 		FScopeLock ScopeLock(&CriticalSection);
-		
+
 		if(FileAReader.IsValid())
 		{
 			FileAReader->Close();
 			FileAReader = nullptr;
 		}
 	}
-	
+
 	[super dealloc];
 }
 
@@ -143,19 +143,19 @@
 {
 	// There should be no need to queue these up - if it turns out we need to do that - then add an ordered queue of loadingRequest objects
 	[loadingRequest retain];
-	
+
 	// Allow this function to return quickly so the resource loader knows the data is probabky coming and doesn't error
 	AsyncLoadingTask = Async(EAsyncExecution::ThreadPool, [self, loadingRequest]()
 	{
 		FScopeLock ScopeLock(&CriticalSection);
-		
+
 		// If the file reader is created on the Apple callback queue then the PakLoader will throw thread errors
 		if (!bInitialized)
 		{
 			FileAReader = MakeShareable( IFileManager::Get().CreateFileReader(*Path) );
 			bInitialized = true;
 		}
-		
+
 		if(FileAReader.IsValid() && !FileAReader->IsError())
 		{
 			// Fill out content information request - if required
@@ -173,39 +173,39 @@
 			{
 				int64 Offset = loadingRequest.dataRequest.requestedOffset;
 				int64 ByteCount = loadingRequest.dataRequest.requestedLength;
-				
+
 				check(Offset >= 0);
 				check(ByteCount > 0);
-				
+
 				if(Offset + ByteCount <= FileAReader->TotalSize())
 				{
 					FileAReader->Seek(Offset);
-					
+
 					// Don't read the whole requested data range at once - the resource loader often asks for very large data sizes
 					// If we feed it (using respondWithData:) in chunks, it decides it has had enough data usually after a few MB,
 					// then it marks the request as cancelled, this is not an error, before issuing a different request at some point later.
 					// This keeps our peak memory usage down and limits the amount of data we are serializing.
-					
+
 					const int64 MaxChunkBytes = 1024 * 1024 * 1; // in single MB chunks
 					while(ByteCount > 0 && !loadingRequest.isCancelled && !FileAReader->IsError())
 					{
 						int64 ChunkByteCount = MIN(MaxChunkBytes, ByteCount);
 						ByteCount -= ChunkByteCount;
 						check(ByteCount >= 0);
-						
+
 						NSMutableData* nsLoadedData = [[NSMutableData alloc] initWithLength:ChunkByteCount];
 						uint8* pMemory = (uint8*)nsLoadedData.mutableBytes;
 						check(pMemory);
-					
+
 						FileAReader->Serialize(pMemory, ChunkByteCount);
-					
+
 						[loadingRequest.dataRequest respondWithData:nsLoadedData];
 						[nsLoadedData release];
 					}
 				}
 			}
 		}
-		
+
 		// Check file reader is not in error state after potential seek and data read operations
 		if(FileAReader.IsValid() && !FileAReader->IsError())
 		{
@@ -254,14 +254,14 @@ public:
 	, AudioSyncSampleTime(FTimespan::MinValue())
 	, VideoSyncSampleTime(FTimespan::MinValue())
 	{}
-	
+
 	virtual ~FAvfMediaSamples()
 	{}
-	
+
 	virtual bool FetchAudio(TRange<FTimespan> TimeRange, TSharedPtr<IMediaAudioSample, ESPMode::ThreadSafe>& OutSample) override
 	{
 		bool bResult = FMediaSamples::FetchAudio(TimeRange, OutSample);
-		
+
 		if(FTimespan::MinValue() == AudioSyncSampleTime && bResult && OutSample.IsValid())
 		{
 			AudioSyncSampleTime = OutSample->GetTime().Time + OutSample->GetDuration();
@@ -269,11 +269,11 @@ public:
 
 		return bResult;
 	}
-	
+
 	virtual bool FetchVideo(TRange<FTimespan> TimeRange, TSharedPtr<IMediaTextureSample, ESPMode::ThreadSafe>& OutSample)
 	{
 		bool bResult = FMediaSamples::FetchVideo(TimeRange, OutSample);
-		
+
 		if(FTimespan::MinValue() == VideoSyncSampleTime && bResult && OutSample.IsValid())
 		{
 			VideoSyncSampleTime = OutSample->GetTime().Time + OutSample->GetDuration();
@@ -281,11 +281,11 @@ public:
 
 		return bResult;
 	}
-	
+
 	void ClearSyncSampleTimes ()			 { AudioSyncSampleTime = VideoSyncSampleTime = FTimespan::MinValue(); }
 	FTimespan GetAudioSyncSampleTime() const { return AudioSyncSampleTime; }
 	FTimespan GetVideoSyncSampleTime() const { return VideoSyncSampleTime; }
-	
+
 private:
 
 	TAtomic<FTimespan> AudioSyncSampleTime;
@@ -305,12 +305,12 @@ FAvfMediaPlayer::FAvfMediaPlayer(IMediaEventSink& InEventSink)
 	Duration = FTimespan::Zero();
 	MediaUrl = FString();
 	ShouldLoop = false;
-    
+
 	MediaHelper = nil;
     MediaPlayer = nil;
 	PlayerItem = nil;
 	MediaResourceLoader = nil;
-		
+
 	bPrerolled = false;
 	bTimeSynced = false;
 	bSeeking = false;
@@ -326,7 +326,7 @@ FAvfMediaPlayer::~FAvfMediaPlayer()
 
 	delete Tracks;
 	Tracks = nullptr;
-	
+
 	delete Samples;
 	Samples = nullptr;
 }
@@ -403,7 +403,7 @@ void FAvfMediaPlayer::OnStatusNotification()
 				{
 					Tracks->Initialize(PlayerItem, Info);
 					EventSink.ReceiveMediaEvent(EMediaEvent::TracksChanged);
-					
+
 					Duration = FTimespan::FromSeconds(CMTimeGetSeconds(PlayerItem.asset.duration));
 					CurrentState = (CurrentState == EMediaState::Closed) ? EMediaState::Stopped : CurrentState;
 
@@ -421,6 +421,11 @@ void FAvfMediaPlayer::OnStatusNotification()
 							// Preroll for playback.
 							[MediaPlayer prerollAtRate:1.0f completionHandler:^(BOOL bFinished)
 							{
+								// For some reason we now get bFinished==NO most of the time for
+								// unknown reasons. According to the documentation failure is caused
+								// if the playback rate or time is changed during prerolling, but we
+								// do neither, so the cause of the failure is unknown.
+							#if 0
 								if (bFinished)
 								{
 									PlayerTasks.Enqueue([=]()
@@ -441,6 +446,22 @@ void FAvfMediaPlayer::OnStatusNotification()
 										EventSink.ReceiveMediaEvent(EMediaEvent::MediaOpenFailed);
 									});
 								}
+							#else
+								PlayerTasks.Enqueue([=]()
+								{
+									if(PlayerItem.status == AVPlayerItemStatusReadyToPlay)
+									{
+										bPrerolled = true;
+										CurrentState = EMediaState::Stopped;
+										EventSink.ReceiveMediaEvent(EMediaEvent::MediaOpened);
+									}
+									else
+									{
+										CurrentState = EMediaState::Error;
+										EventSink.ReceiveMediaEvent(EMediaEvent::MediaOpenFailed);
+									}
+								});
+							#endif
 							}];
 						}
 					}
@@ -514,7 +535,7 @@ void FAvfMediaPlayer::Close()
 
 	CurrentTime = FTimespan::Zero();
 	MediaUrl = FString();
-	
+
 	if (PlayerItem != nil)
 	{
 		if (MediaHelper != nil)
@@ -534,7 +555,7 @@ void FAvfMediaPlayer::Close()
 		[PlayerItem release];
 		PlayerItem = nil;
 	}
-	
+
 	if (MediaHelper != nil)
 	{
 		[MediaHelper release];
@@ -548,7 +569,7 @@ void FAvfMediaPlayer::Close()
 		[MediaPlayer release];
 		MediaPlayer = nil;
 	}
-	
+
 	if(MediaResourceLoader != nil)
 	{
 		// Wait the loading task to finish before releasing the loader.
@@ -556,21 +577,21 @@ void FAvfMediaPlayer::Close()
 		[MediaResourceLoader release];
 		MediaResourceLoader = nil;
 	}
-	
+
 	Tracks->Reset();
 	EventSink.ReceiveMediaEvent(EMediaEvent::TracksChanged);
 
 	CurrentState = EMediaState::Closed;
 	Duration = CurrentTime = FTimespan::Zero();
 	Info.Empty();
-	
+
 	EventSink.ReceiveMediaEvent(EMediaEvent::MediaClosed);
-		
+
 	bPrerolled = false;
 	bSeeking = false;
 
 	CurrentRate = 0.f;
-	
+
 	ClearTimeSync();
 }
 
@@ -641,14 +662,14 @@ bool FAvfMediaPlayer::Open(const FString& Url, const IMediaOptions* /*Options*/)
 	FString Path;
 
 	bool bPakResourceLoading = false;
-	
+
 	if (Url.StartsWith(TEXT("file://")))
 	{
 		// Media Framework doesn't percent encode the URL, so the path portion is just a native file path.
 		// Extract it and then use it create a proper URL.
 		Path = Url.Mid(7);
 		nsMediaUrl = [NSURL fileURLWithPath:Path.GetNSString() isDirectory:NO];
-		
+
 #if !WITH_EDITOR
 		// Check if this file is in a pak file - if it is, AV Player can't load this we need to handle the data load requests manually
 		FPakPlatformFile* PakFileMgr = (FPakPlatformFile*)(FPlatformFileManager::Get().FindPlatformFile(FPakPlatformFile::GetTypeName()));
@@ -686,7 +707,7 @@ bool FAvfMediaPlayer::Open(const FString& Url, const IMediaOptions* /*Options*/)
 		nsMediaUrl = [NSURL fileURLWithPath: FullPath.GetNSString() isDirectory:NO];
 	}
 #endif
-	
+
 	// create player instance
 	MediaUrl = FPaths::GetCleanFilename(Url);
 	MediaPlayer = [[AVPlayer alloc] init];
@@ -696,7 +717,7 @@ bool FAvfMediaPlayer::Open(const FString& Url, const IMediaOptions* /*Options*/)
 		UE_LOG(LogAvfMedia, Error, TEXT("Failed to create instance of an AVPlayer"));
 		return false;
 	}
-	
+
 	MediaPlayer.actionAtItemEnd = AVPlayerActionAtItemEndPause;
 
 	// create player item
@@ -705,16 +726,16 @@ bool FAvfMediaPlayer::Open(const FString& Url, const IMediaOptions* /*Options*/)
 
 	// Use URL asset which gives us resource loading ability if system can't handle the scheme
 	AVURLAsset* urlAsset = [[AVURLAsset alloc] initWithURL:nsMediaUrl options:nil];
-	
+
 	if(bPakResourceLoading)
 	{
 		MediaResourceLoader = [[FAVMediaAssetResourceLoaderDelegate alloc] initWithPath:Path];
 		[urlAsset.resourceLoader setDelegate:MediaResourceLoader queue:dispatch_get_global_queue(QOS_CLASS_DEFAULT, 0)];
 	}
-	
+
 	PlayerItem = [[AVPlayerItem playerItemWithAsset:urlAsset] retain];
 	[urlAsset release];
-	
+
 	if (PlayerItem == nil)
 	{
 		UE_LOG(LogAvfMedia, Error, TEXT("Failed to open player item with Url:"), *Url);
@@ -745,7 +766,7 @@ bool FAvfMediaPlayer::Open(const FString& Url, const IMediaOptions* /*Options*/)
 			NSString *errstr = [[userInfo objectForKey : NSUnderlyingErrorKey] localizedDescription];
 
 			UE_LOG(LogAvfMedia, Warning, TEXT("Failed to load video tracks. [%s]"), *FString(errstr));
-	 
+
 			PlayerTasks.Enqueue([=]()
 			{
 				CurrentState = EMediaState::Error;
@@ -756,7 +777,7 @@ bool FAvfMediaPlayer::Open(const FString& Url, const IMediaOptions* /*Options*/)
 
 	[[NSNotificationCenter defaultCenter] addObserver:MediaHelper selector:@selector(playerItemPlaybackEndReached:) name:AVPlayerItemDidPlayToEndTimeNotification object:PlayerItem];
 	[PlayerItem addObserver:MediaHelper forKeyPath:@"status" options:0 context:PlayerItem];
-	
+
 	MediaPlayer.rate = 0.0;
 	CurrentTime = FTimespan::Zero();
 
@@ -784,7 +805,7 @@ bool FAvfMediaPlayer::Open(const FString& Url, const IMediaOptions* /*Options*/)
 	{
 		AudioRouteChangedHandle = FCoreDelegates::AudioRouteChangedDelegate.AddRaw(this, &FAvfMediaPlayer::HandleAudioRouteChanged);
 	}
-	
+
 	return true;
 }
 
@@ -838,7 +859,7 @@ void FAvfMediaPlayer::TickInput(FTimespan DeltaTime, FTimespan /*Timecode*/)
 						{
 							SyncTime = FTimespan::FromSeconds(CMTimeGetSeconds(MediaPlayer.currentTime));
 						}
-						
+
 						if(SyncTime != FTimespan::MinValue())
 						{
 							bTimeSynced = true;
@@ -868,7 +889,7 @@ void FAvfMediaPlayer::TickInput(FTimespan DeltaTime, FTimespan /*Timecode*/)
 			}
 		}
 	}
-	
+
 	// process deferred tasks
 	TFunction<void()> Task;
 
@@ -962,10 +983,10 @@ bool FAvfMediaPlayer::Seek(const FTimespan& Time)
 		ClearTimeSync();
 
 		CurrentTime = Time;
-		
+
 		double TotalSeconds = Time.GetTotalSeconds();
 		CMTime CurrentTimeInSeconds = CMTimeMakeWithSeconds(TotalSeconds, 1000);
-		
+
 		static CMTime Tolerance = CMTimeMakeWithSeconds(0.01, 1000);
 		[MediaPlayer seekToTime:CurrentTimeInSeconds toleranceBefore:Tolerance toleranceAfter:Tolerance completionHandler:^(BOOL bFinished)
 		{
@@ -987,7 +1008,7 @@ bool FAvfMediaPlayer::Seek(const FTimespan& Time)
 bool FAvfMediaPlayer::SetLooping(bool Looping)
 {
 	ShouldLoop = Looping;
-	
+
 	if (ShouldLoop)
 	{
 		MediaPlayer.actionAtItemEnd = AVPlayerActionAtItemEndNone;
@@ -1004,11 +1025,11 @@ bool FAvfMediaPlayer::SetLooping(bool Looping)
 bool FAvfMediaPlayer::SetRate(float Rate)
 {
 	CurrentRate = Rate;
-	
+
 	if (bPrerolled)
 	{
 		[MediaPlayer setRate : CurrentRate];
-		
+
 		if (FMath::IsNearlyZero(CurrentRate) && CurrentState != EMediaState::Paused)
 		{
 			CurrentState = EMediaState::Paused;
@@ -1019,12 +1040,12 @@ bool FAvfMediaPlayer::SetRate(float Rate)
 			if(CurrentState != EMediaState::Playing)
 			{
 				ClearTimeSync();
-				
+
 				CurrentState = EMediaState::Playing;
 				EventSink.ReceiveMediaEvent(EMediaEvent::PlaybackResumed);
 			}
 		}
-		
+
 		// Use AVPlayer Mute to control reverse playback audio playback
 		// Only needed if !AUDIO_PLAYBACK_VIA_ENGINE - however - keep all platforms the same
 		bool bMuteAudio = Rate < 0.f;
