UE-218169 Xcode 16 build fixes
-Fixed instances where the TEXT() macro was used within a static assert.
-Replaced [=] default this capturing lamdas with explicit [this] capture.

diff --git a/Engine/Plugins/Experimental/Web/WebAPI/Source/WebAPIEditor/Public/WebAPIDefinition.h b/Engine/Plugins/Experimental/Web/WebAPI/Source/WebAPIEditor/Public/WebAPIDefinition.h
--- a/Engine/Plugins/Experimental/Web/WebAPI/Source/WebAPIEditor/Public/WebAPIDefinition.h
+++ b/Engine/Plugins/Experimental/Web/WebAPI/Source/WebAPIEditor/Public/WebAPIDefinition.h
@@ -36,7 +36,7 @@ public:
 	template <typename DataCacheType>
 	DataCacheType* AddOrGetImportedDataCache(FName InKey)
 	{
-		static_assert(TIsDerivedFrom<DataCacheType, UObject>::IsDerived, TEXT("DataCacheType should derive from UObject"));
+		static_assert(TIsDerivedFrom<DataCacheType, UObject>::IsDerived, "DataCacheType should derive from UObject");
 		return Cast<DataCacheType>(AddOrGetImportedDataCache(InKey, DataCacheType::StaticClass()));		
 	}
 
diff --git a/Engine/Plugins/Media/AvfMedia/Source/AvfMedia/Private/Player/AvfMediaPlayer.cpp b/Engine/Plugins/Media/AvfMedia/Source/AvfMedia/Private/Player/AvfMediaPlayer.cpp
--- a/Engine/Plugins/Media/AvfMedia/Source/AvfMedia/Private/Player/AvfMediaPlayer.cpp
+++ b/Engine/Plugins/Media/AvfMedia/Source/AvfMedia/Private/Player/AvfMediaPlayer.cpp
@@ -370,7 +370,7 @@ void FAvfMediaPlayer::OnEndReached()
 {
 	if (ShouldLoop)
 	{
-		PlayerTasks.Enqueue([=]()
+		PlayerTasks.Enqueue([this]()
 		{
 			EventSink.ReceiveMediaEvent(EMediaEvent::PlaybackEndReached);
 			Seek(CurrentRate < 0.f ? Duration : FTimespan::Zero());
@@ -381,7 +381,7 @@ void FAvfMediaPlayer::OnEndReached()
 		CurrentState = EMediaState::Paused;
 		CurrentRate = 0.0f;
 
-		PlayerTasks.Enqueue([=]()
+		PlayerTasks.Enqueue([this]()
 		{
 			Seek(FTimespan::Zero());
 			EventSink.ReceiveMediaEvent(EMediaEvent::PlaybackEndReached);
@@ -393,7 +393,7 @@ void FAvfMediaPlayer::OnEndReached()
 
 void FAvfMediaPlayer::OnStatusNotification()
 {
-	PlayerTasks.Enqueue([=]()
+	PlayerTasks.Enqueue([this]()
 	{
 		switch(PlayerItem.status)
 		{
@@ -428,7 +428,7 @@ void FAvfMediaPlayer::OnStatusNotification()
 							#if 0
 								if (bFinished)
 								{
-									PlayerTasks.Enqueue([=]()
+									PlayerTasks.Enqueue([this]()
 									{
 										if(PlayerItem.status == AVPlayerItemStatusReadyToPlay)
 										{
@@ -440,14 +440,14 @@ void FAvfMediaPlayer::OnStatusNotification()
 								}
 								else
 								{
-									PlayerTasks.Enqueue([=]()
+									PlayerTasks.Enqueue([this]()
 									{
 										CurrentState = EMediaState::Error;
 										EventSink.ReceiveMediaEvent(EMediaEvent::MediaOpenFailed);
 									});
 								}
 							#else
-								PlayerTasks.Enqueue([=]()
+								PlayerTasks.Enqueue([this]()
 								{
 									if(PlayerItem.status == AVPlayerItemStatusReadyToPlay)
 									{
@@ -754,7 +754,7 @@ bool FAvfMediaPlayer::Open(const FString& Url, const IMediaOptions* /*Options*/)
 			// File movies will be ready now
 			if (PlayerItem.status == AVPlayerItemStatusReadyToPlay)
 			{
-				PlayerTasks.Enqueue([=]()
+				PlayerTasks.Enqueue([this]()
 				{
 					OnStatusNotification();
 				});
@@ -767,7 +767,7 @@ bool FAvfMediaPlayer::Open(const FString& Url, const IMediaOptions* /*Options*/)
 
 			UE_LOG(LogAvfMedia, Warning, TEXT("Failed to load video tracks. [%s]"), *FString(errstr));
 
-			PlayerTasks.Enqueue([=]()
+			PlayerTasks.Enqueue([this]()
 			{
 				CurrentState = EMediaState::Error;
 				EventSink.ReceiveMediaEvent(EMediaEvent::MediaOpenFailed);
@@ -992,7 +992,7 @@ bool FAvfMediaPlayer::Seek(const FTimespan& Time)
 		{
 			if(bFinished)
 			{
-				PlayerTasks.Enqueue([=]()
+				PlayerTasks.Enqueue([this]()
 				{
 					bSeeking = false;
 					EventSink.ReceiveMediaEvent(EMediaEvent::SeekCompleted);
diff --git a/Engine/Plugins/Runtime/GameFeatures/Source/GameFeatures/Private/GameFeaturesSubsystem.cpp b/Engine/Plugins/Runtime/GameFeatures/Source/GameFeatures/Private/GameFeaturesSubsystem.cpp
--- a/Engine/Plugins/Runtime/GameFeatures/Source/GameFeatures/Private/GameFeaturesSubsystem.cpp
+++ b/Engine/Plugins/Runtime/GameFeatures/Source/GameFeatures/Private/GameFeaturesSubsystem.cpp
@@ -168,7 +168,7 @@ const FString LexToString(const EGameFeatureTargetState GameFeatureTargetState)
 			return TEXT("Unknown");
 	}
 
-	static_assert((uint8)EGameFeatureTargetState::Count == 4, TEXT("Update LexToString to include new EGameFeatureTargetState"));
+	static_assert((uint8)EGameFeatureTargetState::Count == 4, "Update LexToString to include new EGameFeatureTargetState");
 }
 
 void LexFromString(EGameFeatureTargetState& ValueOut, const TCHAR* StringIn)
diff --git a/Engine/Plugins/Runtime/MassGameplay/Source/MassRepresentation/Public/MassRepresentationTypes.h b/Engine/Plugins/Runtime/MassGameplay/Source/MassRepresentation/Public/MassRepresentationTypes.h
--- a/Engine/Plugins/Runtime/MassGameplay/Source/MassRepresentation/Public/MassRepresentationTypes.h
+++ b/Engine/Plugins/Runtime/MassGameplay/Source/MassRepresentation/Public/MassRepresentationTypes.h
@@ -204,7 +204,7 @@ private:
 	// prevalent use in FMassRepresentationFragment. Perhaps serial number could be formed from the referenced 
 	// FStaticMeshInstanceVisualizationDesc's hash.
 };
-static_assert(sizeof(FStaticMeshInstanceVisualizationDescHandle) == sizeof(uint16), TEXT("FStaticMeshInstanceVisualizationDescHandle must be uint16 sized to ensure FMassRepresentationFragment memory isn't unexpectedly bloated"));
+static_assert(sizeof(FStaticMeshInstanceVisualizationDescHandle) == sizeof(uint16), "FStaticMeshInstanceVisualizationDescHandle must be uint16 sized to ensure FMassRepresentationFragment memory isn't unexpectedly bloated");
 
 class UInstancedStaticMeshComponent;
 
diff --git a/Engine/Source/Runtime/Engine/Private/SoundWaveProcedural.cpp b/Engine/Source/Runtime/Engine/Private/SoundWaveProcedural.cpp
--- a/Engine/Source/Runtime/Engine/Private/SoundWaveProcedural.cpp
+++ b/Engine/Source/Runtime/Engine/Private/SoundWaveProcedural.cpp
@@ -16,7 +16,7 @@ USoundWaveProcedural::USoundWaveProcedural(const FObjectInitializer& ObjectIniti
 	bReset = false;
 	NumBufferUnderrunSamples = 512;
 	NumSamplesToGeneratePerCallback = DEFAULT_PROCEDURAL_SOUNDWAVE_BUFFER_SIZE;
-	static_assert(DEFAULT_PROCEDURAL_SOUNDWAVE_BUFFER_SIZE >= 512, TEXT("Should generate more samples than this per callback."));
+	static_assert(DEFAULT_PROCEDURAL_SOUNDWAVE_BUFFER_SIZE >= 512, "Should generate more samples than this per callback.");
 	//checkf(NumSamplesToGeneratePerCallback >= NumBufferUnderrunSamples, TEXT("Should generate more samples than this per callback."));
 
 	// If the main audio device has been set up, we can use this to define our callback size.
