[avf] fix FString::Printf error

diff --git a/Engine/Plugins/Media/AvfMedia/Source/AvfMedia/Private/Player/AvfMediaTracks.cpp b/Engine/Plugins/Media/AvfMedia/Source/AvfMedia/Private/Player/AvfMediaTracks.cpp
--- a/Engine/Plugins/Media/AvfMedia/Source/AvfMedia/Private/Player/AvfMediaTracks.cpp
+++ b/Engine/Plugins/Media/AvfMedia/Source/AvfMedia/Private/Player/AvfMediaTracks.cpp
@@ -359,7 +359,7 @@ void FAvfMediaTracks::AppendStats(FString &OutStats) const
 			const FTrack& Track = VideoTracks[TrackIndex];
 
 			OutStats += FString::Printf(TEXT("    %s\n"), *GetTrackDisplayName(EMediaTrackType::Video, TrackIndex).ToString());
-			OutStats += FString::Printf(TEXT("        BitRate: %i\n"), [Track.AssetTrack estimatedDataRate]);
+			OutStats += FString::Printf(TEXT("        BitRate: %i\n"), int([Track.AssetTrack estimatedDataRate]));
 		}
 	}
 }
