Fix the memory leak issue on mediaserverd process

diff --git a/Engine/Plugins/Media/AvfMedia/Source/AvfMedia/Private/Player/AvfMediaPlayer.cpp b/Engine/Plugins/Media/AvfMedia/Source/AvfMedia/Private/Player/AvfMediaPlayer.cpp
--- a/Engine/Plugins/Media/AvfMedia/Source/AvfMedia/Private/Player/AvfMediaPlayer.cpp
+++ b/Engine/Plugins/Media/AvfMedia/Source/AvfMedia/Private/Player/AvfMediaPlayer.cpp
@@ -198,8 +198,9 @@
 						check(pMemory);
 
 						FileAReader->Serialize(pMemory, ChunkByteCount);
-
-						[loadingRequest.dataRequest respondWithData:nsLoadedData];
+						@autoreleasepool{
+							[loadingRequest.dataRequest respondWithData:nsLoadedData];
+						}
 						[nsLoadedData release];
 					}
 				}
@@ -209,7 +210,9 @@
 		// Check file reader is not in error state after potential seek and data read operations
 		if(FileAReader.IsValid() && !FileAReader->IsError())
 		{
-			[loadingRequest finishLoading];
+			@autoreleasepool{
+				[loadingRequest finishLoading];
+			}
 		}
 		else
 		{
@@ -551,7 +554,7 @@ void FAvfMediaPlayer::Close()
 		[PlayerItem.asset loadValuesAsynchronouslyForKeys:@[@"tracks"] completionHandler:^
 		{
 		}];
-
+		MediaPlayer.rate = 0.0;
 		[PlayerItem release];
 		PlayerItem = nil;
 	}
