Fixed #include using wrong casing for headers included

#Synced-CL 50214969

diff --git a/Engine/Plugins/Media/AvfMedia/Source/AvfMediaCapture/Private/Player/AvfMediaCaptureHelper.h b/Engine/Plugins/Media/AvfMedia/Source/AvfMediaCapture/Private/Player/AvfMediaCaptureHelper.h
--- a/Engine/Plugins/Media/AvfMedia/Source/AvfMediaCapture/Private/Player/AvfMediaCaptureHelper.h
+++ b/Engine/Plugins/Media/AvfMedia/Source/AvfMediaCapture/Private/Player/AvfMediaCaptureHelper.h
@@ -11,7 +11,7 @@
 #include "Apple/PreAppleSystemHeaders.h"
 #include <AVFoundation/AVFoundation.h>
 #if PLATFORM_MAC
-	#include <Appkit/Appkit.h>
+	#include <AppKit/AppKit.h>
 #endif
 #include "Apple/PostAppleSystemHeaders.h"
 
diff --git a/Engine/Source/Runtime/Launch/Private/Mac/LaunchMac.cpp b/Engine/Source/Runtime/Launch/Private/Mac/LaunchMac.cpp
--- a/Engine/Source/Runtime/Launch/Private/Mac/LaunchMac.cpp
+++ b/Engine/Source/Runtime/Launch/Private/Mac/LaunchMac.cpp
@@ -8,7 +8,7 @@
 #include "HAL/ExceptionHandling.h"
 #include "Mac/MacPlatformCrashContext.h"
 #include "Mac/MacPlatformMiscEx.h"
-#include "Apple/LaunchCppToSwift.h"
+#include "Apple/LaunchCPPToSwift.h"
 
 #if WITH_ENGINE
 	#include "Engine/Engine.h"
