[AvfMedia] Provide CPU buffer alongside GPU texture when using FAvfMediaCapturePlayer

Fix double sRGB to linear conversion.

#Synced-CL 52165013

diff --git a/Engine/Plugins/Media/AvfMedia/Source/AvfMedia/Private/AvfMediaTextureSample.cpp b/Engine/Plugins/Media/AvfMedia/Source/AvfMedia/Private/AvfMediaTextureSample.cpp
new file mode 100644
--- /dev/null
+++ b/Engine/Plugins/Media/AvfMedia/Source/AvfMedia/Private/AvfMediaTextureSample.cpp
@@ -0,0 +1,158 @@
+// Copyright Epic Games, Inc. All Rights Reserved.
+
+#include "AvfMediaTextureSample.h"
+#include "AvfMediaPrivate.h"
+
+FAvfMediaTextureSample::~FAvfMediaTextureSample()
+{
+	UnlockAndReleasePixelBuffer();
+}
+
+#if WITH_ENGINE
+
+bool FAvfMediaTextureSample::Initialize(
+	TRefCountPtr<FRHITexture>& InTexture,
+	const FIntPoint& InDim,
+	const FIntPoint& InOutputDim,
+	FTimespan InTime,
+	FTimespan InDuration,
+	CVPixelBufferRef InPixelBuffer)
+{
+	if (!InTexture.IsValid())
+	{
+		return false;
+	}
+
+	Dim = InDim;
+	Duration = InDuration;
+	OutputDim = InOutputDim;
+	Texture = InTexture;
+	Time = InTime;
+
+	UnlockAndReleasePixelBuffer();
+
+	if (InPixelBuffer)
+	{
+		if (CVPixelBufferLockBaseAddress(InPixelBuffer, kCVPixelBufferLock_ReadOnly) == kCVReturnSuccess)
+		{
+			CFRetain(InPixelBuffer);
+			LockedPixelBuffer = InPixelBuffer;
+		}
+	}
+
+	return true;
+}
+
+#else // WITH_ENGINE
+
+bool FAvfMediaTextureSample::Initialize(
+	uint8* InBuffer,
+	const FIntPoint& InDim,
+	const FIntPoint& InOutputDim,
+	uint32 InStride,
+	FTimespan InTime,
+	FTimespan InDuration)
+{
+	if (InBuffer == nullptr)
+	{
+		return false;
+	}
+
+	const int32 BufferSize = InDim.X * InDim.Y * 4;
+
+	if (BufferSize == 0)
+	{
+		return false;
+	}
+
+	Buffer.Reset(BufferSize);
+	Buffer.Append((uint8*)InBuffer, BufferSize);
+
+	Dim = InDim;
+	Duration = InDuration;
+	OutputDim = InOutputDim;
+	Stride = InStride;
+	Time = InTime;
+
+	return true;
+}
+
+#endif // WITH_ENGINE
+
+
+const void* FAvfMediaTextureSample::GetBuffer()
+{
+	return LockedPixelBuffer ? CVPixelBufferGetBaseAddress(LockedPixelBuffer) : nullptr;
+}
+
+FIntPoint FAvfMediaTextureSample::GetDim() const
+{
+	return Dim;
+}
+
+FTimespan FAvfMediaTextureSample::GetDuration() const
+{
+	return Duration;
+}
+
+EMediaTextureSampleFormat FAvfMediaTextureSample::GetFormat() const
+{
+	return EMediaTextureSampleFormat::CharBGRA;
+}
+
+FIntPoint FAvfMediaTextureSample::GetOutputDim() const
+{
+	return OutputDim;
+}
+
+uint32 FAvfMediaTextureSample::GetStride() const
+{
+	return LockedPixelBuffer ? CVPixelBufferGetBytesPerRow(LockedPixelBuffer) : Dim.X * 4;
+}
+
+FRHITexture* FAvfMediaTextureSample::GetTexture() const
+{
+	return Texture;
+}
+
+FMediaTimeStamp FAvfMediaTextureSample::GetTime() const
+{
+	return FMediaTimeStamp(Time);
+}
+
+bool FAvfMediaTextureSample::IsCacheable() const
+{
+	return true;
+}
+
+bool FAvfMediaTextureSample::IsOutputSrgb() const
+{
+	// The GPU texture is created as MTLPixelFormatBGRA8Unorm_sRGB, which tells the RHI
+	// to apply sRGB-to-linear conversion on texture read. IsOutputSrgb() is a separate
+	// flag consumed by UMediaTexture to determine whether it should ALSO request sRGB
+	// conversion when creating render targets from this sample. Returning true here would
+	// cause a double conversion (too dark). Return false since the texture format already
+	// encodes the sRGB intent at the RHI level.
+	return false;
+}
+
+void FAvfMediaTextureSample::ShutdownPoolable()
+{
+	UnlockAndReleasePixelBuffer();
+}
+
+void FAvfMediaTextureSample::Reset()
+{
+	UnlockAndReleasePixelBuffer();
+	Texture = nullptr;
+}
+
+void FAvfMediaTextureSample::UnlockAndReleasePixelBuffer()
+{
+	if (LockedPixelBuffer)
+	{
+		CVPixelBufferUnlockBaseAddress(LockedPixelBuffer, kCVPixelBufferLock_ReadOnly);
+		CFRelease(LockedPixelBuffer);
+		LockedPixelBuffer = nullptr;
+	}
+}
diff --git a/Engine/Plugins/Media/AvfMedia/Source/AvfMedia/Public/AvfMediaTextureSample.h b/Engine/Plugins/Media/AvfMedia/Source/AvfMedia/Public/AvfMediaTextureSample.h
--- a/Engine/Plugins/Media/AvfMedia/Source/AvfMedia/Public/AvfMediaTextureSample.h
+++ b/Engine/Plugins/Media/AvfMedia/Source/AvfMedia/Public/AvfMediaTextureSample.h
@@ -6,9 +6,7 @@
 #include "IMediaTextureSample.h"
 #include "MediaObjectPool.h"
 #include "Math/IntPoint.h"
-#include "Math/Range.h"
 #include "Misc/Timespan.h"
-#include "Containers/ResourceArray.h"
 
 #if WITH_ENGINE
 	#include "RHI.h"
@@ -16,18 +14,31 @@
 	#include "Containers/Array.h"
 #endif
 
+/** CVPixelBuffer.h on macOS pulls in CoreServices -> CarbonCore -> NumberFormatting.h
+ * which defines its own struct FVector, clashing with UE's type alias. Temporarily
+ * redirect FVector to a harmless name while the Apple header is parsed.
+ */
+#if PLATFORM_MAC
+	#pragma push_macro("FVector")
+	#define FVector FVectorWorkaround
+#endif
+#include <CoreVideo/CVPixelBuffer.h>
+#if PLATFORM_MAC
+	#undef FVector
+	#pragma pop_macro("FVector")
+#endif
 
 /**
  * Texture sample generated by AvfMedia player.
  */
-class FAvfMediaTextureSample
+class DLLEXPORT FAvfMediaTextureSample
 	: public IMediaTextureSample
 	, public IMediaPoolable
 {
 public:
 
 	/** Virtual destructor. */
-	virtual ~FAvfMediaTextureSample() = default;
+	virtual ~FAvfMediaTextureSample();
 
 public:
 
@@ -41,6 +52,7 @@ public:
 	 * @param InOutputDim The sample's output width and height (in pixels).
 	 * @param InTime The sample time (in the player's local clock).
 	 * @param InDuration The duration for which the sample is valid.
+	 * @param InPixelBuffer Optional CVPixelBuffer to retain for CPU access via GetBuffer().
 	 * @return true on success, false otherwise.
 	 */
 	bool Initialize(
@@ -48,21 +60,8 @@ public:
 		const FIntPoint& InDim,
 		const FIntPoint& InOutputDim,
 		FTimespan InTime,
-		FTimespan InDuration)
-	{
-		if (!InTexture.IsValid())
-		{
-			return false;
-		}
-
-		Dim = InDim;
-		Duration = InDuration;
-		OutputDim = InOutputDim;
-		Texture = InTexture;
-		Time = InTime;
-
-		return true;
-	}
+		FTimespan InDuration,
+		CVPixelBufferRef InPixelBuffer = nullptr);
 
 #else
 
@@ -83,94 +82,34 @@ public:
 		const FIntPoint& InOutputDim,
 		uint32 InStride,
 		FTimespan InTime,
-		FTimespan InDuration)
-	{
-		if (InBuffer == nullptr)
-		{
-			return false;
-		}
-
-		const int32 BufferSize = InDim.X * InDim.Y * 4;
-
-		if (BufferSize == 0)
-		{
-			return false;
-		}
-
-		Buffer.Reset(BufferSize);
-		Buffer.Append((uint8*)InBuffer, BufferSize);
-
-		Dim = InDim;
-		Duration = InDuration;
-		OutputDim = InOutputDim;
-		Stride = InStride;
-		Time = InTime;
-
-		return true;
-	}
+		FTimespan InDuration);
+
 #endif
 
 public:
 
 	//~ IMediaTextureSample interface
 
-	virtual const void* GetBuffer() override
-	{
-		return nullptr;
-	}
-
-	virtual FIntPoint GetDim() const override
-	{
-		return Texture->GetTexture2D()->GetSizeXY();
-	}
-
-	virtual FTimespan GetDuration() const override
-	{
-		return Duration;
-	}
-
-	virtual EMediaTextureSampleFormat GetFormat() const override
-	{
-		return EMediaTextureSampleFormat::CharBGRA;
-	}
-
-	virtual FIntPoint GetOutputDim() const override
-	{
-		return Texture->GetTexture2D()->GetSizeXY();
-	}
-
-	virtual uint32 GetStride() const override
-	{
-		return Texture->GetTexture2D()->GetSizeX() * 4;
-	}
-
-	virtual FRHITexture* GetTexture() const override
-	{
-		return Texture;
-	}
-
-	virtual FMediaTimeStamp GetTime() const override
-	{
-		return FMediaTimeStamp(Time);
-	}
-
-	virtual bool IsCacheable() const override
-	{
-		return true;
-	}
-
-	virtual bool IsOutputSrgb() const override
-	{
-		return true;
-	}
-
-	virtual void Reset() override
-	{
-		Texture = nullptr;
-	}
-	
+	virtual const void* GetBuffer() override;
+	virtual FIntPoint GetDim() const override;
+	virtual FTimespan GetDuration() const override;
+	virtual EMediaTextureSampleFormat GetFormat() const override;
+	virtual FIntPoint GetOutputDim() const override;
+	virtual uint32 GetStride() const override;
+	virtual FRHITexture* GetTexture() const override;
+	virtual FMediaTimeStamp GetTime() const override;
+	virtual bool IsCacheable() const override;
+	virtual bool IsOutputSrgb() const override;
+	virtual void Reset() override;
+
+	//~ IMediaPoolable interface
+
+	virtual void ShutdownPoolable() override;
+
 private:
 
+	void UnlockAndReleasePixelBuffer();
+
 	/** Width and height of the texture sample. */
 	FIntPoint Dim;
 
@@ -195,6 +134,9 @@ private:
 	uint32 Stride;
 
 #endif
+
+	/** Retained CVPixelBuffer for CPU access (locked read-only while sample is active). */
+	CVPixelBufferRef LockedPixelBuffer = nullptr;
 };
 
 
diff --git a/Engine/Plugins/Media/AvfMedia/Source/AvfMediaCapture/AvfMediaCapture.Build.cs b/Engine/Plugins/Media/AvfMedia/Source/AvfMediaCapture/AvfMediaCapture.Build.cs
--- a/Engine/Plugins/Media/AvfMedia/Source/AvfMediaCapture/AvfMediaCapture.Build.cs
+++ b/Engine/Plugins/Media/AvfMedia/Source/AvfMediaCapture/AvfMediaCapture.Build.cs
@@ -21,7 +21,8 @@ namespace UnrealBuildTool.Rules
 					"MediaUtils",
 					"RenderCore",
 					"RHI",
-					"MetalRHI"
+					"MetalRHI",
+					"AvfMedia"
 				});
 
 			AddEngineThirdPartyPrivateStaticDependencies(Target, "MetalCPP");
@@ -31,10 +32,6 @@ namespace UnrealBuildTool.Rules
 					"Media",
 				});
 
-			PublicIncludePathModuleNames.AddRange(
-				new string[] {
-					"AvfMedia",
-				});
 
 			PublicFrameworks.AddRange(
 				new string[] {
diff --git a/Engine/Plugins/Media/AvfMedia/Source/AvfMediaCapture/Private/Player/AvfMediaCapturePlayer.cpp b/Engine/Plugins/Media/AvfMedia/Source/AvfMediaCapture/Private/Player/AvfMediaCapturePlayer.cpp
--- a/Engine/Plugins/Media/AvfMedia/Source/AvfMediaCapture/Private/Player/AvfMediaCapturePlayer.cpp
+++ b/Engine/Plugins/Media/AvfMedia/Source/AvfMediaCapture/Private/Player/AvfMediaCapturePlayer.cpp
@@ -743,8 +743,12 @@ void FAvfMediaCapturePlayer::ProcessSampleBufferVideo(CMSampleBufferRef SampleBu
 
 		FIntPoint Dim(Width, Height);
 
+		// Retain the pixel buffer so it survives past this callback. The render command
+		// will pass it to FAvfMediaTextureSample which takes its own retain, then we release ours.
+		CFRetain(PixelBuffer);
+
 		ENQUEUE_RENDER_COMMAND(FAvfMediaCapturePlayer_ProcessSampleBufferVideo_CreateTexture)(
-            [WeakPlayer = this->AsWeak(), Dim, SampleDuration, TextureRef](FRHICommandListImmediate& RHICmdList)
+            [WeakPlayer = this->AsWeak(), Dim, SampleDuration, TextureRef, PixelBuffer](FRHICommandListImmediate& RHICmdList)
             {
             	if (const TSharedPtr<FAvfMediaCapturePlayer> Player = WeakPlayer.Pin())
             	{
@@ -753,12 +757,18 @@ void FAvfMediaCapturePlayer::ProcessSampleBufferVideo(CMSampleBufferRef SampleBu
 					TRefCountPtr<FRHITexture> ShaderResource = MetalDynamicRHI->RHICreateTexture2DFromCVMetalTexture(PF_B8G8R8A8, ETextureCreateFlags::NoTiling | ETextureCreateFlags::ShaderResource, FClearValueBinding::White, TextureRef);
 					CFRelease(TextureRef);
 					
-					
 					TSharedRef<FAvfMediaTextureSample, ESPMode::ThreadSafe> VideoSample = Player->VideoSamplePool.AcquireShared();
-					VideoSample->Initialize(ShaderResource, Dim, Dim, Player->CurrentTime, SampleDuration);
+					VideoSample->Initialize(ShaderResource, Dim, Dim, Player->CurrentTime, SampleDuration, PixelBuffer);
+					CFRelease(PixelBuffer);
             
 					Player->MediaSamples->AddVideo(VideoSample);
             	}
+            	else
+            	{
+            		// Player was destroyed — still need to release our retains
+            		CFRelease(TextureRef);
+            		CFRelease(PixelBuffer);
+            	}
         });
 	}
 }
